Copying //UE5/Dev-Cooker@14539516 to Main (//UE5/Main)

diff --git a/Engine/Config/BaseEngine.ini b/Engine/Config/BaseEngine.ini
--- a/Engine/Config/BaseEngine.ini
+++ b/Engine/Config/BaseEngine.ini
@@ -932,6 +932,10 @@ HangDuration=0.0
 AssetLogShowsDiskPath=True
 AssetLogShowsAbsolutePath=True
 
+[Core.ContentVirtualization]
+EnablePushToBackend=False
+MinPayloadLength=0
+
 [/Script/Engine.StreamingSettings]
 s.MinBulkDataSizeForAsyncLoading=131072
 s.AsyncLoadingThreadEnabled=False
diff --git a/Engine/Plugins/Developer/PerforceSourceControl/Source/PerforceSourceControl/Private/PerforceSourceControlProvider.cpp b/Engine/Plugins/Developer/PerforceSourceControl/Source/PerforceSourceControl/Private/PerforceSourceControlProvider.cpp
--- a/Engine/Plugins/Developer/PerforceSourceControl/Source/PerforceSourceControl/Private/PerforceSourceControlProvider.cpp
+++ b/Engine/Plugins/Developer/PerforceSourceControl/Source/PerforceSourceControl/Private/PerforceSourceControlProvider.cpp
@@ -535,6 +535,8 @@ void FPerforceSourceControlProvider::RegisterWorker( const FName& InName, const
 
 ECommandResult::Type FPerforceSourceControlProvider::ExecuteSynchronousCommand(FPerforceSourceControlCommand& InCommand, const FText& Task, bool bSuppressResponseMsg)
 {
+	TRACE_CPUPROFILER_EVENT_SCOPE(FPerforceSourceControlProvider::ExecuteSynchronousCommand);
+
 	ECommandResult::Type Result = ECommandResult::Failed;
 
 	struct Local
diff --git a/Engine/Plugins/Enterprise/DatasmithContent/Source/DatasmithContentEditor/Private/DatasmithSceneDetails.cpp b/Engine/Plugins/Enterprise/DatasmithContent/Source/DatasmithContentEditor/Private/DatasmithSceneDetails.cpp
--- a/Engine/Plugins/Enterprise/DatasmithContent/Source/DatasmithContentEditor/Private/DatasmithSceneDetails.cpp
+++ b/Engine/Plugins/Enterprise/DatasmithContent/Source/DatasmithContentEditor/Private/DatasmithSceneDetails.cpp
@@ -31,7 +31,7 @@ void FDatasmithSceneDetails::CustomizeDetails(IDetailLayoutBuilder& DetailBuilde
 	check( DatasmithScene );
 	DatasmithScenePtr = DatasmithScene;
 
-	static_assert( TIsSame< TRemovePointer< decltype( UDatasmithScene::AssetImportData ) >::Type, UDatasmithSceneImportData >::Value, "Please update this details customization" );
+	static_assert( TIsSame< typename TRemovePointer< typename TRemoveObjectPointer<decltype( UDatasmithScene::AssetImportData )>::Type >::Type, UDatasmithSceneImportData >::Value, "Please update this details customization" );
 	static_assert( TIsDerivedFrom < UDatasmithSceneImportData, UAssetImportData >::IsDerived, "Please update this details customization" );
 
 	TSharedRef< IPropertyHandle > AssetImportDataPropertyHandle = DetailBuilder.GetProperty( GET_MEMBER_NAME_CHECKED( UDatasmithScene, AssetImportData ) );
diff --git a/Engine/Plugins/Enterprise/MDLImporter/Source/MDLImporter/Private/generator/GeneratorMaterialExpressions.cpp b/Engine/Plugins/Enterprise/MDLImporter/Source/MDLImporter/Private/generator/GeneratorMaterialExpressions.cpp
--- a/Engine/Plugins/Enterprise/MDLImporter/Source/MDLImporter/Private/generator/GeneratorMaterialExpressions.cpp
+++ b/Engine/Plugins/Enterprise/MDLImporter/Source/MDLImporter/Private/generator/GeneratorMaterialExpressions.cpp
@@ -394,6 +394,19 @@ namespace Generator
 		}
 	}
 
+	void CheckedConnect(UObject* Parent, const FMaterialExpressionConnection& Connection, FExpressionInput& Input, TObjectPtr<UTexture>* OutTexturePtr)
+	{
+		if (Connection.ConnectionType == EConnectionType::Texture)
+		{
+			check(OutTexturePtr && Connection.Texture);
+			*OutTexturePtr = Connection.Texture;
+		}
+		else
+		{
+			CheckedConnect(Parent, Connection, Input);
+		}
+	}
+
 	uint32 ComponentCount(const FMaterialExpressionConnection& Input)
 	{
 		switch (Input.ConnectionType)
diff --git a/Engine/Plugins/FX/Niagara/Source/NiagaraEditorWidgets/Private/SNiagaraOverviewStack.cpp b/Engine/Plugins/FX/Niagara/Source/NiagaraEditorWidgets/Private/SNiagaraOverviewStack.cpp
--- a/Engine/Plugins/FX/Niagara/Source/NiagaraEditorWidgets/Private/SNiagaraOverviewStack.cpp
+++ b/Engine/Plugins/FX/Niagara/Source/NiagaraEditorWidgets/Private/SNiagaraOverviewStack.cpp
@@ -184,7 +184,7 @@ class SNiagaraSystemOverviewEntryListRow : public STableRow<UNiagaraStackEntry*>
 			{
 				if (ModuleFunctionCall.FunctionScript->IsAsset() || GbShowNiagaraDeveloperWindows > 0)
 				{
-					GEditor->GetEditorSubsystem<UAssetEditorSubsystem>()->OpenEditorForAsset(const_cast<UNiagaraScript*>(ModuleFunctionCall.FunctionScript));
+					GEditor->GetEditorSubsystem<UAssetEditorSubsystem>()->OpenEditorForAsset(const_cast<UNiagaraScript*>(ToRawPtr(ModuleFunctionCall.FunctionScript)));
 					return FReply::Handled();
 				}
 				else if (ModuleItem->IsScratchModule())
diff --git a/Engine/Plugins/FX/Niagara/Source/NiagaraEditorWidgets/Private/Stack/SNiagaraStackModuleItem.cpp b/Engine/Plugins/FX/Niagara/Source/NiagaraEditorWidgets/Private/Stack/SNiagaraStackModuleItem.cpp
--- a/Engine/Plugins/FX/Niagara/Source/NiagaraEditorWidgets/Private/Stack/SNiagaraStackModuleItem.cpp
+++ b/Engine/Plugins/FX/Niagara/Source/NiagaraEditorWidgets/Private/Stack/SNiagaraStackModuleItem.cpp
@@ -59,7 +59,7 @@ FReply SNiagaraStackModuleItem::OnMouseButtonDoubleClick(const FGeometry& InMyGe
 	{
 		if (ModuleFunctionCall.FunctionScript->IsAsset() || GbShowNiagaraDeveloperWindows > 0)
 		{
-			GEditor->GetEditorSubsystem<UAssetEditorSubsystem>()->OpenEditorForAsset(const_cast<UNiagaraScript*>(ModuleFunctionCall.FunctionScript));
+			GEditor->GetEditorSubsystem<UAssetEditorSubsystem>()->OpenEditorForAsset(const_cast<UNiagaraScript*>(ToRawPtr(ModuleFunctionCall.FunctionScript)));
 			return FReply::Handled();
 		}
 		else if (ModuleItem->IsScratchModule())
diff --git a/Engine/Plugins/Lumin/MagicLeap/Source/MagicLeapAR/Private/LuminARTrackingSystem.h b/Engine/Plugins/Lumin/MagicLeap/Source/MagicLeapAR/Private/LuminARTrackingSystem.h
--- a/Engine/Plugins/Lumin/MagicLeap/Source/MagicLeapAR/Private/LuminARTrackingSystem.h
+++ b/Engine/Plugins/Lumin/MagicLeap/Source/MagicLeapAR/Private/LuminARTrackingSystem.h
@@ -145,7 +145,7 @@ void FLuminARImplementation::RemoveHandleIf(TUniqueFunction<bool(const FGuid&)>
 	for (TMap<FGuid, FTrackedGeometryGroup>::TIterator Iterator = TrackedGeometryGroups.CreateIterator(); Iterator; ++Iterator)
 	{
 		FTrackedGeometryGroup& TrackedGeoemtryGroup = Iterator->Value;
-		auto* DerivedPointer = Cast<T>(TrackedGeoemtryGroup.TrackedGeometry);
+		auto* DerivedPointer = Cast<T>(ToRawPtr(TrackedGeoemtryGroup.TrackedGeometry));
 		if (DerivedPointer && HandleOperator(Iterator->Key))
 		{
 			DerivedPointer->SetTrackingState(EARTrackingState::StoppedTracking);
diff --git a/Engine/Plugins/MeshPainting/Source/MeshPaintingToolset/Public/MeshPaintingToolsetTypes.h b/Engine/Plugins/MeshPainting/Source/MeshPaintingToolset/Public/MeshPaintingToolsetTypes.h
--- a/Engine/Plugins/MeshPainting/Source/MeshPaintingToolset/Public/MeshPaintingToolsetTypes.h
+++ b/Engine/Plugins/MeshPainting/Source/MeshPaintingToolset/Public/MeshPaintingToolsetTypes.h
@@ -96,7 +96,13 @@ struct FPaintableTexture
 	UTexture*	Texture;
 	int32		UVChannelIndex;
 
-	FPaintableTexture(UTexture* InTexture = nullptr, uint32 InUVChannelIndex = 0)
+	FPaintableTexture()
+		: Texture(nullptr)
+		, UVChannelIndex(0)
+	{}
+
+	template<typename T, decltype(ImplicitConv<UTexture*>(DeclVal<T>()))* = nullptr>
+	FPaintableTexture(T InTexture = nullptr, uint32 InUVChannelIndex = 0)
 		: Texture(InTexture)
 		, UVChannelIndex(InUVChannelIndex)
 	{}
diff --git a/Engine/Plugins/MovieScene/CustomizableSequencerTracks/Source/CustomizableSequencerTracks/Private/SequencerSectionBP.cpp b/Engine/Plugins/MovieScene/CustomizableSequencerTracks/Source/CustomizableSequencerTracks/Private/SequencerSectionBP.cpp
--- a/Engine/Plugins/MovieScene/CustomizableSequencerTracks/Source/CustomizableSequencerTracks/Private/SequencerSectionBP.cpp
+++ b/Engine/Plugins/MovieScene/CustomizableSequencerTracks/Source/CustomizableSequencerTracks/Private/SequencerSectionBP.cpp
@@ -32,7 +32,7 @@ void USequencerSectionBP::ImportEntityImpl(UMovieSceneEntitySystemLinker* Entity
 		FGuid ObjectBindingID = Params.GetObjectBindingID();
 		OutImportedEntity->AddBuilder(
 			FEntityBuilder()
-			.Add(BuiltInComponents->TrackInstance, FMovieSceneTrackInstanceComponent{ this, CustomTrack->TrackInstanceType })
+			.Add(BuiltInComponents->TrackInstance, FMovieSceneTrackInstanceComponent{ decltype(FMovieSceneTrackInstanceComponent::Owner)(this), CustomTrack->TrackInstanceType })
 			.AddConditional(BuiltInComponents->GenericObjectBinding, ObjectBindingID, ObjectBindingID.IsValid())
 		);
 	}
diff --git a/Engine/Plugins/NetcodeUnitTest/NetcodeUnitTest/Source/NetcodeUnitTest/Private/UnitTestCommandlet.cpp b/Engine/Plugins/NetcodeUnitTest/NetcodeUnitTest/Source/NetcodeUnitTest/Private/UnitTestCommandlet.cpp
--- a/Engine/Plugins/NetcodeUnitTest/NetcodeUnitTest/Source/NetcodeUnitTest/Private/UnitTestCommandlet.cpp
+++ b/Engine/Plugins/NetcodeUnitTest/NetcodeUnitTest/Source/NetcodeUnitTest/Private/UnitTestCommandlet.cpp
@@ -35,7 +35,7 @@
 
 
 // Enable access to the private UGameEngine.GameInstance value, using the GET_PRIVATE macro
-IMPLEMENT_GET_PRIVATE_VAR(UGameEngine, GameInstance, UGameInstance*);
+IMPLEMENT_GET_PRIVATE_VAR(UGameEngine, GameInstance, decltype(UGameEngine::GameInstance));
 
 /** The exit-confirmation dialog, displayed to the user when all unit tests are complete */
 static TSharedPtr<SWindow> ConfirmDialog=NULL;
diff --git a/Engine/Plugins/Runtime/GameplayAbilities/Source/GameplayAbilities/Private/GameplayEffect.cpp b/Engine/Plugins/Runtime/GameplayAbilities/Source/GameplayAbilities/Private/GameplayEffect.cpp
--- a/Engine/Plugins/Runtime/GameplayAbilities/Source/GameplayAbilities/Private/GameplayEffect.cpp
+++ b/Engine/Plugins/Runtime/GameplayAbilities/Source/GameplayAbilities/Private/GameplayEffect.cpp
@@ -33,11 +33,11 @@ DECLARE_CYCLE_STAT(TEXT("MakeQuery"), STAT_MakeGameplayEffectQuery, STATGROUP_Ab
 
 #if WITH_EDITOR
 #define GETCURVE_REPORTERROR_WITHPOSTLOAD(Handle) \
-	if (Handle.CurveTable) const_cast<UCurveTable*>(Handle.CurveTable)->ConditionalPostLoad(); \
+	if (Handle.CurveTable) const_cast<UCurveTable*>(ToRawPtr(Handle.CurveTable))->ConditionalPostLoad(); \
 	GETCURVE_REPORTERROR(Handle);
 
 #define GETCURVE_REPORTERROR_WITHPATHNAME_WITHPOSTLOAD(Handle, PathNameString) \
-	if (Handle.CurveTable) const_cast<UCurveTable*>(Handle.CurveTable)->ConditionalPostLoad(); \
+	if (Handle.CurveTable) const_cast<UCurveTable*>(ToRawPtr(Handle.CurveTable))->ConditionalPostLoad(); \
 	GETCURVE_REPORTERROR_WITHPATHNAME(Handle, PathNameString);
 #endif // WITH_EDITOR
 
diff --git a/Engine/Programs/UnrealHeaderTool/Config/DefaultEngine.ini b/Engine/Programs/UnrealHeaderTool/Config/DefaultEngine.ini
--- a/Engine/Programs/UnrealHeaderTool/Config/DefaultEngine.ini
+++ b/Engine/Programs/UnrealHeaderTool/Config/DefaultEngine.ini
@@ -31,3 +31,7 @@ gc.MaxObjectsInProgram=262144
 +TypeRedirects=(OldType="TAssetPtr", NewType="TSoftObjectPtr")
 +TypeRedirects=(OldType="TAssetSubclassOf", NewType="TSoftClassPtr")
 
+[UnrealHeaderTool]
+; This will increase the size of the UHT log files, but the size should return
+; to normal after native pointer upgrades are completed
+NativePointerMemberBehavior=AllowAndLog
diff --git a/Engine/Source/Editor/AnimGraph/Classes/AnimGraphNode_Base.h b/Engine/Source/Editor/AnimGraph/Classes/AnimGraphNode_Base.h
--- a/Engine/Source/Editor/AnimGraph/Classes/AnimGraphNode_Base.h
+++ b/Engine/Source/Editor/AnimGraph/Classes/AnimGraphNode_Base.h
@@ -265,6 +265,8 @@ class ANIMGRAPH_API UAnimGraphNode_Base : public UK2Node
 	// Helper function for ReplaceReferredAnimations	
 	template<class AssetType>
 	void HandleAnimReferenceReplacement(AssetType*& OriginalAsset, const TMap<UAnimationAsset*, UAnimationAsset*>& AnimAssetReplacementMap);
+	template<class AssetType>
+	void HandleAnimReferenceReplacement(TObjectPtr<AssetType>& OriginalAsset, const TMap<UAnimationAsset*, UAnimationAsset*>& AnimAssetReplacementMap);
 
 	/**
 	 * Selection notification callback.
@@ -402,3 +404,9 @@ void UAnimGraphNode_Base::HandleAnimReferenceReplacement(AssetType*& OriginalAss
 		OriginalAsset = Cast<AssetType>(*ReplacementAsset);
 	}
 }
+
+template<class AssetType>
+void UAnimGraphNode_Base::HandleAnimReferenceReplacement(TObjectPtr<AssetType>& OriginalAsset, const TMap<UAnimationAsset*, UAnimationAsset*>& AnimAssetReplacementMap)
+{
+	HandleAnimReferenceReplacement(static_cast<AssetType*&>(OriginalAsset), AnimAssetReplacementMap);
+}
diff --git a/Engine/Source/Editor/BlueprintGraph/Classes/BlueprintNodeBinder.h b/Engine/Source/Editor/BlueprintGraph/Classes/BlueprintNodeBinder.h
--- a/Engine/Source/Editor/BlueprintGraph/Classes/BlueprintNodeBinder.h
+++ b/Engine/Source/Editor/BlueprintGraph/Classes/BlueprintNodeBinder.h
@@ -18,7 +18,8 @@ public:
 	FBindingObject()
 		: bIsUObject(false)
 	{}
-	FBindingObject(UObject* InObject)
+	template <typename T, decltype(ImplicitConv<UObject*>(DeclVal<T>()))* = nullptr>
+	FBindingObject(T InObject)
 		: Object(InObject)
 		, bIsUObject(true)
 	{}
@@ -38,9 +39,10 @@ public:
 			Property = InFieldOrObject.ToField();
 		}
 	}
-	FBindingObject& operator=(UObject* InObject)
+	template <typename T, decltype(ImplicitConv<UObject*>(DeclVal<T>()))* = nullptr>
+	FBindingObject& operator=(T InObject)
 	{
-		Object = InObject;
+		Object = ImplicitConv<UObject*>(InObject);
 		Property = nullptr;
 		bIsUObject = true;
 		return *this;
diff --git a/Engine/Source/Editor/BlueprintGraph/Private/BlueprintEventNodeSpawner.cpp b/Engine/Source/Editor/BlueprintGraph/Private/BlueprintEventNodeSpawner.cpp
--- a/Engine/Source/Editor/BlueprintGraph/Private/BlueprintEventNodeSpawner.cpp
+++ b/Engine/Source/Editor/BlueprintGraph/Private/BlueprintEventNodeSpawner.cpp
@@ -128,7 +128,7 @@ FBlueprintNodeSignature UBlueprintEventNodeSpawner::GetSpawnerSignature() const
 	}
 	else
 	{
-		SpawnerSignature.AddSubObject(const_cast<UFunction*>(EventFunc));
+		SpawnerSignature.AddSubObject(const_cast<UFunction*>(ToRawPtr(EventFunc)));
 	}
 	return SpawnerSignature;
 }
@@ -203,7 +203,7 @@ UEdGraphNode* UBlueprintEventNodeSpawner::Invoke(UEdGraph* ParentGraph, FBinding
 			UserDelegate.ExecuteIfBound(NewNode, bInIsTemplateNode);
 		};
 
-		FCustomizeNodeDelegate PostSpawnDelegate = FCustomizeNodeDelegate::CreateStatic(PostSpawnLambda, EventFunc, EventName, CustomizeNodeDelegate);
+		FCustomizeNodeDelegate PostSpawnDelegate = FCustomizeNodeDelegate::CreateStatic(PostSpawnLambda, ToRawPtr(EventFunc), EventName, CustomizeNodeDelegate);
 		EventNode = Super::SpawnNode<UK2Node_Event>(NodeClass, ParentGraph, Bindings, Location, PostSpawnDelegate);
 	}
 	// else, a node for this event already exists, and we should return that 
diff --git a/Engine/Source/Editor/DetailCustomizations/Private/WorldSettingsDetails.cpp b/Engine/Source/Editor/DetailCustomizations/Private/WorldSettingsDetails.cpp
--- a/Engine/Source/Editor/DetailCustomizations/Private/WorldSettingsDetails.cpp
+++ b/Engine/Source/Editor/DetailCustomizations/Private/WorldSettingsDetails.cpp
@@ -379,7 +379,7 @@ void FLightmapCustomNodeBuilder::RefreshLightmapItems()
 	if ( World )
 	{
 		TArray<UTexture2D*> LightMapsAndShadowMaps;
-		World->GetLightMapsAndShadowMaps(World->GetCurrentLevel(), LightMapsAndShadowMaps);
+		World->GetLightMapsAndShadowMaps(World->GetCurrentLevel(), LightMapsAndShadowMaps, false);
 
 		for ( auto ObjIt = LightMapsAndShadowMaps.CreateConstIterator(); ObjIt; ++ObjIt )
 		{
diff --git a/Engine/Source/Editor/GraphEditor/Private/SGraphEditorImpl.cpp b/Engine/Source/Editor/GraphEditor/Private/SGraphEditorImpl.cpp
--- a/Engine/Source/Editor/GraphEditor/Private/SGraphEditorImpl.cpp
+++ b/Engine/Source/Editor/GraphEditor/Private/SGraphEditorImpl.cpp
@@ -895,7 +895,7 @@ void SGraphEditorImpl::AddContextMenuCommentSection(UToolMenu* InMenu)
 					FToolMenuSection& Section = InMenu->AddSection("GraphNodeComment", LOCTEXT("NodeCommentMenuHeader", "Node Comment"));
 					Section.AddEntry(FToolMenuEntry::InitWidget("NodeCommentBox", NodeCommentBox, FText::GetEmpty()));
 				}
-				TWeakObjectPtr<UEdGraphNode> SelectedNodeWeakPtr = MakeWeakObjectPtr(const_cast<UEdGraphNode*>(Context->Node));
+				TWeakObjectPtr<UEdGraphNode> SelectedNodeWeakPtr = MakeWeakObjectPtr(const_cast<UEdGraphNode*>(ToRawPtr(Context->Node)));
 
 				FText NodeCommentText;
 				if (UEdGraphNode* SelectedNode = SelectedNodeWeakPtr.Get())
@@ -943,7 +943,7 @@ void SGraphEditorImpl::AddContextMenuCommentSection(UToolMenu* InMenu)
 					LOCTEXT("MultiCommentDesc", "Create Comment from Selection"),
 					LOCTEXT("CommentToolTip", "Create a resizable comment box around selection."),
 					FSlateIcon(),
-					FExecuteAction::CreateStatic(SCommentUtility::CreateComment, GraphSchema, const_cast<UEdGraph*>(Context->Graph)
+					FExecuteAction::CreateStatic(SCommentUtility::CreateComment, GraphSchema, const_cast<UEdGraph*>(ToRawPtr(Context->Graph))
 				));
 			}
 		}
diff --git a/Engine/Source/Editor/MaterialEditor/Private/MaterialEditor.cpp b/Engine/Source/Editor/MaterialEditor/Private/MaterialEditor.cpp
--- a/Engine/Source/Editor/MaterialEditor/Private/MaterialEditor.cpp
+++ b/Engine/Source/Editor/MaterialEditor/Private/MaterialEditor.cpp
@@ -989,7 +989,7 @@ void FMaterialEditor::CollapseNodesIntoGraph(UEdGraphNode* InGatewayNode, UMater
 						UMaterialExpressionPinBase* PinBase = CastChecked<UMaterialExpressionPinBase>(LocalPort->MaterialExpression);
 
 						GatewayRerouteExpression->SubgraphExpression = CompositeNode->MaterialExpression;
-						PinBase->ReroutePins.Add(FCompositeReroute{UniquePortName, GatewayReroutePin});
+						PinBase->ReroutePins.Add(FCompositeReroute{UniquePortName, decltype(FCompositeReroute::Expression)(GatewayReroutePin)});
 						PinBase->Modify();
 					}
 				}
@@ -6316,9 +6316,9 @@ void FMaterialEditor::DeepCopyExpressions(UMaterialGraph* CopyGraph, UMaterialEx
 				NewPinBase->ReroutePins.Empty();
 				for (FCompositeReroute& Reroute : OldPinBase->ReroutePins)
 				{
-					UMaterialExpressionReroute* DupReroute = DuplicateExpression(Reroute.Expression);
+					UMaterialExpressionReroute* DupReroute = DuplicateExpression(ToRawPtr(Reroute.Expression));
 					DupReroute->SubgraphExpression = NewSubgraphExpression;
-					NewPinBase->ReroutePins.Add({ Reroute.Name, DupReroute });
+					NewPinBase->ReroutePins.Add({ Reroute.Name, decltype(FCompositeReroute::Expression)(DupReroute) });
 				}
 
 				UMaterialExpressionComposite* SubGraphComposite = CastChecked<UMaterialExpressionComposite>(NewSubgraphExpression);
@@ -6333,7 +6333,7 @@ void FMaterialEditor::DeepCopyExpressions(UMaterialGraph* CopyGraph, UMaterialEx
 			}
 			else
 			{
-				MaterialNode->MaterialExpression = DuplicateExpression(MaterialNode->MaterialExpression);
+				MaterialNode->MaterialExpression = DuplicateExpression(ToRawPtr(MaterialNode->MaterialExpression));
 				MaterialNode->MaterialExpression->SubgraphExpression = NewSubgraphExpression;
 			}
 			
@@ -6346,7 +6346,7 @@ void FMaterialEditor::DeepCopyExpressions(UMaterialGraph* CopyGraph, UMaterialEx
 		{
 			Comment->Modify();
 			Comment->Rename(/*NewName=*/ NULL, /*NewOuter=*/ CopyGraph);
-			Comment->MaterialExpressionComment = DuplicateExpression(Comment->MaterialExpressionComment);
+			Comment->MaterialExpressionComment = DuplicateExpression(ToRawPtr(Comment->MaterialExpressionComment));
 			Comment->MaterialExpressionComment->SubgraphExpression = NewSubgraphExpression;
 
 			// GraphNode is transient so it won't be duplicated. 
diff --git a/Engine/Source/Editor/MaterialEditor/Private/MaterialEditorInstanceDetailCustomization.cpp b/Engine/Source/Editor/MaterialEditor/Private/MaterialEditorInstanceDetailCustomization.cpp
--- a/Engine/Source/Editor/MaterialEditor/Private/MaterialEditorInstanceDetailCustomization.cpp
+++ b/Engine/Source/Editor/MaterialEditor/Private/MaterialEditorInstanceDetailCustomization.cpp
@@ -321,7 +321,7 @@ void FMaterialInstanceParameterDetails::CreateGroupsWidget(TSharedRef<IPropertyH
 		if (!MaterialEditorInstance->bIsFunctionPreviewMaterial)
 		{
 			OnChildButtonClicked = FOnClicked::CreateStatic(&FMaterialPropertyHelpers::OnClickedSaveNewMaterialInstance, LocalSourceInstance, LocalEditorInstance);
-			OnSiblingButtonClicked = FOnClicked::CreateStatic(&FMaterialPropertyHelpers::OnClickedSaveNewMaterialInstance, MaterialEditorInstance->SourceInstance->Parent, LocalEditorInstance);
+			OnSiblingButtonClicked = FOnClicked::CreateStatic(&FMaterialPropertyHelpers::OnClickedSaveNewMaterialInstance, ToRawPtr(MaterialEditorInstance->SourceInstance->Parent), LocalEditorInstance);
 		}
 		else
 		{
diff --git a/Engine/Source/Editor/MaterialEditor/Private/MaterialInstanceEditor.cpp b/Engine/Source/Editor/MaterialEditor/Private/MaterialInstanceEditor.cpp
--- a/Engine/Source/Editor/MaterialEditor/Private/MaterialInstanceEditor.cpp
+++ b/Engine/Source/Editor/MaterialEditor/Private/MaterialInstanceEditor.cpp
@@ -1538,7 +1538,7 @@ bool FMaterialInstanceEditor::ApproveSetPreviewAsset(UObject* InAsset)
 void FMaterialInstanceEditor::Refresh()
 {
 	int32 TempIndex;
-	const bool bParentChanged = !MaterialParentList.Find( MaterialEditorInstance->Parent, TempIndex );
+	const bool bParentChanged = !MaterialParentList.Find( ToRawPtr(MaterialEditorInstance->Parent), TempIndex );
 
 	PreviewVC->RefreshViewport();
 
diff --git a/Engine/Source/Editor/MaterialEditor/Private/MaterialStats.h b/Engine/Source/Editor/MaterialEditor/Private/MaterialStats.h
--- a/Engine/Source/Editor/MaterialEditor/Private/MaterialStats.h
+++ b/Engine/Source/Editor/MaterialEditor/Private/MaterialStats.h
@@ -438,7 +438,7 @@ public:
 	/** this will set the material to be analyzed by this class */
 	/** TMaterial should be UMaterial or UMaterialInstance */
 	template <typename TMaterial>
-	void SetMaterial(TMaterial *InMaterial);
+	void SetMaterial(TMaterial& InMaterial);
 
 	// end Update Functions
 	///////////////////////////////////////////
@@ -512,7 +512,7 @@ FORCEINLINE TSharedPtr<class SMaterialEditorStatsWidget> FMaterialStats::GetGrid
 
 //TMaterial should be UMaterial or UMaterialInstance
 template <typename TMaterial>
-void FMaterialStats::SetMaterial(TMaterial *MaterialPtr)
+void FMaterialStats::SetMaterial(TMaterial& MaterialPtr)
 {
 	if (MaterialInterface != MaterialPtr)
 	{
diff --git a/Engine/Source/Editor/MaterialEditor/Private/SMaterialLayersFunctionsTree.cpp b/Engine/Source/Editor/MaterialEditor/Private/SMaterialLayersFunctionsTree.cpp
--- a/Engine/Source/Editor/MaterialEditor/Private/SMaterialLayersFunctionsTree.cpp
+++ b/Engine/Source/Editor/MaterialEditor/Private/SMaterialLayersFunctionsTree.cpp
@@ -210,6 +210,11 @@ bool SMaterialLayersFunctionsInstanceTree::IsOverriddenExpression(class UDEditor
 	return FMaterialPropertyHelpers::IsOverriddenExpression(Parameter) && FunctionInstance->LayerStates[InIndex];
 }
 
+bool SMaterialLayersFunctionsInstanceTree::IsOverriddenExpression(TObjectPtr<UDEditorParameterValue> Parameter, int32 InIndex)
+{
+	return IsOverriddenExpression(Parameter.Get(), InIndex);
+}
+
 FGetShowHiddenParameters SMaterialLayersFunctionsInstanceTree::GetShowHiddenDelegate() const
 {
 	return ShowHiddenDelegate;
@@ -221,6 +226,11 @@ void  SMaterialLayersFunctionsInstanceTreeItem::OnOverrideParameter(bool NewValu
 	Tree->GetWrapper()->Refresh();
 }
 
+void  SMaterialLayersFunctionsInstanceTreeItem::OnOverrideParameter(bool NewValue, TObjectPtr<UDEditorParameterValue> Parameter)
+{
+	OnOverrideParameter(NewValue, Parameter.Get());
+}
+
 void SMaterialLayersFunctionsInstanceTreeItem::Construct(const FArguments& InArgs, const TSharedRef<STableViewBase>& InOwnerTableView)
 {
 	StackParameterData = InArgs._StackParameterData;
diff --git a/Engine/Source/Editor/MaterialEditor/Private/SMaterialLayersFunctionsTree.h b/Engine/Source/Editor/MaterialEditor/Private/SMaterialLayersFunctionsTree.h
--- a/Engine/Source/Editor/MaterialEditor/Private/SMaterialLayersFunctionsTree.h
+++ b/Engine/Source/Editor/MaterialEditor/Private/SMaterialLayersFunctionsTree.h
@@ -17,8 +17,9 @@
 #include "PropertyCustomizationHelpers.h"
 
 class IPropertyHandle;
-class UMaterialEditorInstanceConstant;
 class SMaterialLayersFunctionsInstanceTree;
+class UDEditorParameterValue;
+class UMaterialEditorInstanceConstant;
 
 class SMaterialLayersFunctionsInstanceTreeItem : public STableRow< TSharedPtr<FSortedParamData> >
 {
@@ -74,6 +75,7 @@ public:
 
 	FReply OnLayerDrop(const FDragDropEvent& DragDropEvent);
 	void OnOverrideParameter(bool NewValue, class UDEditorParameterValue* Parameter);
+	void OnOverrideParameter(bool NewValue, TObjectPtr<UDEditorParameterValue> Parameter);
 	/**
 	* Construct the widget
 	*
@@ -166,6 +168,7 @@ public:
 	void UpdateThumbnailMaterial(TEnumAsByte<EMaterialParameterAssociation> InAssociation, int32 InIndex, bool bAlterBlendIndex = false);
 	FReply OnThumbnailDoubleClick(const FGeometry& Geometry, const FPointerEvent& MouseEvent, EMaterialParameterAssociation InAssociation, int32 InIndex);
 	bool IsOverriddenExpression(class UDEditorParameterValue* Parameter, int32 InIndex);
+	bool IsOverriddenExpression(TObjectPtr<UDEditorParameterValue> Parameter, int32 InIndex);
 
 	FGetShowHiddenParameters GetShowHiddenDelegate() const;
 protected:
diff --git a/Engine/Source/Editor/MaterialEditor/Public/MaterialPropertyHelpers.h b/Engine/Source/Editor/MaterialEditor/Public/MaterialPropertyHelpers.h
--- a/Engine/Source/Editor/MaterialEditor/Public/MaterialPropertyHelpers.h
+++ b/Engine/Source/Editor/MaterialEditor/Public/MaterialPropertyHelpers.h
@@ -17,6 +17,7 @@ struct FAssetData;
 class IDetailGroup;
 class IDetailLayoutBuilder;
 class IPropertyHandle;
+class UDEditorParameterValue;
 class UMaterialEditorInstanceConstant;
 enum class ECheckBoxState : uint8;
 class UMaterialInterface;
@@ -140,6 +141,7 @@ class MATERIALEDITOR_API FMaterialPropertyHelpers
 public:
 	/** Returns true if the parameter is being overridden */
 	static bool IsOverriddenExpression(class UDEditorParameterValue* Parameter);
+	static bool IsOverriddenExpression(TObjectPtr<UDEditorParameterValue> Parameter) { return IsOverriddenExpression(Parameter.Get()); }
 	static ECheckBoxState IsOverriddenExpressionCheckbox(UDEditorParameterValue* Parameter);
 
 	/** Gets the expression description of this parameter from the base material */
@@ -152,21 +154,26 @@ public:
 	static void OnOverrideParameter(bool NewValue, class UDEditorParameterValue* Parameter, UMaterialEditorInstanceConstant* MaterialEditorInstance);
 
 	static EVisibility ShouldShowExpression(UDEditorParameterValue* Parameter, UMaterialEditorInstanceConstant* MaterialEditorInstance, FGetShowHiddenParameters ShowHiddenDelegate);
+	static EVisibility ShouldShowExpression(TObjectPtr<UDEditorParameterValue> Parameter, UMaterialEditorInstanceConstant* MaterialEditorInstance, FGetShowHiddenParameters ShowHiddenDelegate) { return ShouldShowExpression(Parameter.Get(), MaterialEditorInstance, ShowHiddenDelegate); }
 
 	/** Generic material property reset to default implementation.  Resets Parameter to default */
 	static void ResetToDefault(TSharedPtr<IPropertyHandle> PropertyHandle, class UDEditorParameterValue* Parameter, UMaterialEditorInstanceConstant* MaterialEditorInstance);
+	static void ResetToDefault(TSharedPtr<IPropertyHandle> PropertyHandle, TObjectPtr<UDEditorParameterValue> Parameter, UMaterialEditorInstanceConstant* MaterialEditorInstance) { ResetToDefault(PropertyHandle, Parameter.Get(), MaterialEditorInstance); }
 	static bool ShouldShowResetToDefault(TSharedPtr<IPropertyHandle> PropertyHandle, class UDEditorParameterValue* Parameter, UMaterialEditorInstanceConstant* MaterialEditorInstance);
+	static bool ShouldShowResetToDefault(TSharedPtr<IPropertyHandle> PropertyHandle, TObjectPtr<UDEditorParameterValue> Parameter, UMaterialEditorInstanceConstant* MaterialEditorInstance) { return ShouldShowResetToDefault(PropertyHandle, Parameter.Get(), MaterialEditorInstance); }
 	
 	/** Specific resets for layer and blend asses */
 	static void ResetLayerAssetToDefault(TSharedPtr<IPropertyHandle> PropertyHandle, class UDEditorParameterValue* InParameter, TEnumAsByte<EMaterialParameterAssociation> InAssociation, int32 Index, UMaterialEditorInstanceConstant* MaterialEditorInstance);
 	/** If reset to default button should show for a layer or blend asset*/
 	static bool ShouldLayerAssetShowResetToDefault(TSharedPtr<IPropertyHandle> PropertyHandle, TSharedPtr<struct FSortedParamData> InParameterData, UMaterialInterface* InMaterial);
+	static bool ShouldLayerAssetShowResetToDefault(TSharedPtr<IPropertyHandle> PropertyHandle, TSharedPtr<struct FSortedParamData> InParameterData, TObjectPtr<UMaterialInterface> InMaterial) { return ShouldLayerAssetShowResetToDefault(PropertyHandle, InParameterData, InMaterial.Get()); }
 
 	static void OnMaterialLayerAssetChanged(const struct FAssetData& InAssetData, int32 Index, EMaterialParameterAssociation MaterialType, TSharedPtr<class IPropertyHandle> InHandle, FMaterialLayersFunctions* InMaterialFunction);
 
 	static bool FilterLayerAssets(const struct FAssetData& InAssetData, FMaterialLayersFunctions* LayerFunction, EMaterialParameterAssociation MaterialType, int32 Index);
 
 	static FReply OnClickedSaveNewMaterialInstance(class UMaterialInterface* Object, UObject* EditorObject);
+	static FReply OnClickedSaveNewMaterialInstance(TObjectPtr<class UMaterialInterface> Object, UObject* EditorObject) { return OnClickedSaveNewMaterialInstance(Object.Get(), EditorObject); }
 
 	static void CopyMaterialToInstance(class UMaterialInstanceConstant* ChildInstance, TArray<struct FEditorParameterGroup> &ParameterGroups);
 	static void TransitionAndCopyParameters(class UMaterialInstanceConstant* ChildInstance, TArray<struct FEditorParameterGroup> &ParameterGroups, bool bForceCopy = false);
@@ -175,7 +182,9 @@ public:
 
 	static void GetVectorChannelMaskComboBoxStrings(TArray<TSharedPtr<FString>>& OutComboBoxStrings, TArray<TSharedPtr<class SToolTip>>& OutToolTips, TArray<bool>& OutRestrictedItems);
 	static FString GetVectorChannelMaskValue(class UDEditorParameterValue* InParameter);
+	static FString GetVectorChannelMaskValue(TObjectPtr<UDEditorParameterValue> InParameter) { return GetVectorChannelMaskValue(InParameter.Get()); }
 	static void SetVectorChannelMaskValue(const FString& StringValue, TSharedPtr<IPropertyHandle> PropertyHandle, class UDEditorParameterValue* InParameter, UObject* MaterialEditorInstance);
+	static void SetVectorChannelMaskValue(const FString& StringValue, TSharedPtr<IPropertyHandle> PropertyHandle, TObjectPtr<UDEditorParameterValue> InParameter, UObject* MaterialEditorInstance) { SetVectorChannelMaskValue(StringValue, PropertyHandle, InParameter.Get(), MaterialEditorInstance); }
 
 	static TArray<class UFactory*> GetAssetFactories(EMaterialParameterAssociation AssetType);
 	/**
@@ -192,6 +201,7 @@ public:
 	static void SetPositionFromCurveAsset(const FAssetData& AssetData, TSoftObjectPtr<class UCurveLinearColorAtlas> InAtlas, class UDEditorScalarParameterValue* InParameter, TSharedPtr<IPropertyHandle> PropertyHandle, UObject* MaterialEditorInstance);
 
 	static void ResetCurveToDefault(TSharedPtr<IPropertyHandle> PropertyHandle, class UDEditorParameterValue* Parameter, UMaterialEditorInstanceConstant* MaterialEditorInstance);
+	static void ResetCurveToDefault(TSharedPtr<IPropertyHandle> PropertyHandle, TObjectPtr<UDEditorParameterValue> Parameter, UMaterialEditorInstanceConstant* MaterialEditorInstance) { ResetCurveToDefault(PropertyHandle, Parameter.Get(), MaterialEditorInstance); }
 
 	static FText LayerID;
 	static FText BlendID;
diff --git a/Engine/Source/Editor/MeshPaint/Public/MeshPaintTypes.h b/Engine/Source/Editor/MeshPaint/Public/MeshPaintTypes.h
--- a/Engine/Source/Editor/MeshPaint/Public/MeshPaintTypes.h
+++ b/Engine/Source/Editor/MeshPaint/Public/MeshPaintTypes.h
@@ -81,7 +81,13 @@ struct FPaintableTexture
 	UTexture*	Texture;
 	int32		UVChannelIndex;
 
-	FPaintableTexture(UTexture* InTexture = nullptr, uint32 InUVChannelIndex = 0)
+	FPaintableTexture()
+		: Texture(nullptr)
+		, UVChannelIndex(0)
+	{}
+
+	template<typename T, decltype(ImplicitConv<UTexture*>(DeclVal<T>()))* = nullptr>
+	FPaintableTexture(T InTexture, uint32 InUVChannelIndex = 0)
 		: Texture(InTexture)
 		, UVChannelIndex(InUVChannelIndex)
 	{}
diff --git a/Engine/Source/Editor/Persona/Private/PersonaMeshDetails.h b/Engine/Source/Editor/Persona/Private/PersonaMeshDetails.h
--- a/Engine/Source/Editor/Persona/Private/PersonaMeshDetails.h
+++ b/Engine/Source/Editor/Persona/Private/PersonaMeshDetails.h
@@ -639,6 +639,7 @@ private:
 	void OnPreviewMeshChanged(USkeletalMesh* OldSkeletalMesh, USkeletalMesh* NewMesh);
 	
 	bool FilterOutBakePose(const struct FAssetData& AssetData, USkeleton* Skeleton) const;
+	bool FilterOutBakePose(const struct FAssetData& AssetData, TObjectPtr<USkeleton> Skeleton) const { return FilterOutBakePose(AssetData, Skeleton.Get()); }
 
 	FText GetLODCustomModeNameContent(int32 LODIndex) const;
 	ECheckBoxState IsLODCustomModeCheck(int32 LODIndex) const;
diff --git a/Engine/Source/Editor/UMGEditor/Public/WidgetBlueprintEditor.h b/Engine/Source/Editor/UMGEditor/Public/WidgetBlueprintEditor.h
--- a/Engine/Source/Editor/UMGEditor/Public/WidgetBlueprintEditor.h
+++ b/Engine/Source/Editor/UMGEditor/Public/WidgetBlueprintEditor.h
@@ -284,6 +284,9 @@ private:
 	/** Add an animation track for the supplied slot to the current animation. */
 	void AddSlotTrack( UPanelSlot* Slot );
 
+	/** Add an animation track for the supplied slot to the current animation. */
+	void AddSlotTrack(TObjectPtr<UPanelSlot> Slot) { return AddSlotTrack(Slot.Get()); }
+
 	/** Add an animation track for the supplied material property path to the current animation. */
 	void AddMaterialTrack( UWidget* Widget, TArray<FProperty*> MaterialPropertyPath, FText MaterialPropertyDisplayName );
 
diff --git a/Engine/Source/Editor/UnrealEd/Classes/Commandlets/ResavePackagesCommandlet.h b/Engine/Source/Editor/UnrealEd/Classes/Commandlets/ResavePackagesCommandlet.h
--- a/Engine/Source/Editor/UnrealEd/Classes/Commandlets/ResavePackagesCommandlet.h
+++ b/Engine/Source/Editor/UnrealEd/Classes/Commandlets/ResavePackagesCommandlet.h
@@ -3,6 +3,7 @@
 #pragma once
 
 #include "CoreMinimal.h"
+#include "Templates/PimplPtr.h"
 #include "UObject/ObjectMacros.h"
 #include "Commandlets/Commandlet.h"
 #include "Engine/EngineTypes.h"
@@ -71,6 +72,9 @@ protected:
 	/** if we should auto checkout packages that need to be saved**/
 	bool bAutoCheckOut;
 
+	/** if we should batch together source control operations rather than submit them one at a time */
+	bool bBatchSourceControl;
+
 	/** if we should simply skip checked out files rather than error-ing out */
 	bool bSkipCheckedOutFiles;
 
@@ -122,6 +126,7 @@ protected:
 	/** Running count of packages that got modified and will need to be resaved */
 	int32 PackagesConsideredForResave;
 	int32 PackagesResaved;
+	int32 PackagesDeleted;
 
 	/** Only collect garbage after N packages */
 	int32 GarbageCollectionFrequency;
@@ -138,6 +143,9 @@ protected:
 	/** List of redirector packages that should be fixed up at the end */
 	TArray<FString> RedirectorsToFixup;
 
+	/** A queue containing source control operations to be performed in batches. */
+	TPimplPtr<class FQueuedSourceControlOperations> SourceControlQueue;
+
 	/**
 	 * Evaluates the command-line to determine which maps to check.  By default all maps are checked
 	 * Provides child classes with a chance to initialize any variables, parse the command line, etc.
@@ -149,6 +157,8 @@ protected:
 	 */
 	virtual int32 InitializeResaveParameters( const TArray<FString>& Tokens, TArray<FString>& MapPathNames );
 
+	void ParseSourceControlOptions(const TArray<FString>& Tokens);
+
 	/** Loads and saves a single package */
 	virtual void LoadAndSaveOnePackage(const FString& Filename);
 
@@ -212,6 +222,26 @@ protected:
 	void CheckoutAndSavePackage(UPackage* Package, TArray<FString>& SublevelFilenames, bool bIgnoreAlreadyCheckedOut = false);
 	void CheckInFiles(const TArray<FString>& InFilesToSubmit, const FText& InDescription) const;
 
+	/**
+	 * Creates and returns a unique filename in the temporary file directory.
+	 *
+	 * @return	A unique temporary filename
+	 */
+	FString CreateTempFilename();
+
+	/**
+	 * Returns the path of the directory we use to store all temporary files for
+	 * this commandlet.
+	 *
+	 * @return Path to the temporary file directory
+	 */
+	FString GetTempFilesDirectory();
+
+	/**
+	 * Delete the remaining files in the temporary files directory for the commandlet.
+	 */
+	void CleanTempFiles();	
+
 	// Print out a message only if running in very verbose mode
 	void VerboseMessage(const FString& Message);
 
diff --git a/Engine/Source/Editor/UnrealEd/Private/Commandlets/ContentCommandlets.cpp b/Engine/Source/Editor/UnrealEd/Private/Commandlets/ContentCommandlets.cpp
--- a/Engine/Source/Editor/UnrealEd/Private/Commandlets/ContentCommandlets.cpp
+++ b/Engine/Source/Editor/UnrealEd/Private/Commandlets/ContentCommandlets.cpp
@@ -41,6 +41,7 @@
 #include "Engine/Brush.h"
 #include "Editor.h"
 #include "FileHelpers.h"
+#include "CommandletSourceControlUtils.h"
 
 #include "PackageHelperFunctions.h"
 #include "PackageTools.h"
@@ -97,11 +98,13 @@ UResavePackagesCommandlet::UResavePackagesCommandlet(const FObjectInitializer& O
 	: Super(ObjectInitializer)
 	, bForceUATEnvironmentVariableSet(false)
 {
-}
 
+}
 
 int32 UResavePackagesCommandlet::InitializeResaveParameters( const TArray<FString>& Tokens, TArray<FString>& PackageNames )
 {
+	TRACE_CPUPROFILER_EVENT_SCOPE(UResavePackagesCommandlet::InitializeResaveParameters);
+
 	Verbosity = VERY_VERBOSE;
 
 	TArray<FString> Unused;
@@ -118,8 +121,15 @@ int32 UResavePackagesCommandlet::InitializeResaveParameters( const TArray<FStrin
 		if( FParse::Value( *CurrentSwitch, TEXT( "PACKAGE="), Package ) )
 		{
 			FString PackageFile;
-			FPackageName::SearchForPackageOnDisk( Package, NULL, &PackageFile );
-			PackageNames.Add( *PackageFile );
+			if (FPackageName::SearchForPackageOnDisk(Package, NULL, &PackageFile))
+			{
+				PackageNames.Add( *PackageFile );
+			}
+			else
+			{
+				UE_LOG(LogContentCommandlet, Error, TEXT("Failed to find the package given by the cmdline: PACKAGE=%s"), *Package);
+			}
+			
 			bExplicitPackages = true;
 		}
 		else if( FParse::Value( *CurrentSwitch, TEXT( "PACKAGEFOLDER="), PackageFolder ) )
@@ -134,10 +144,6 @@ int32 UResavePackagesCommandlet::InitializeResaveParameters( const TArray<FStrin
 			}
 			bExplicitPackages = true;
 		}
-		else if (FParse::Value(*CurrentSwitch, TEXT("GCFREQ="), GarbageCollectionFrequency))
-		{
-			UE_LOG(LogContentCommandlet, Display, TEXT("Setting garbage collection to happen every %d packages."), GarbageCollectionFrequency);
-		}
 		else if (FParse::Value(*CurrentSwitch, TEXT("MAP="), Maps))
 		{
 			// Allow support for -MAP=Value1+Value2+Value3
@@ -192,7 +198,15 @@ int32 UResavePackagesCommandlet::InitializeResaveParameters( const TArray<FStrin
 				UE_LOG(LogContentCommandlet, Error, TEXT("Failed to load file %s"), *File);
 			}
 		}
+	}
 
+	// Check for numeric settings
+	for (const FString& CurrentSwitch : Switches)
+	{
+		if (FParse::Value(*CurrentSwitch, TEXT("GCFREQ="), GarbageCollectionFrequency))
+		{
+			UE_LOG(LogContentCommandlet, Display, TEXT("Setting garbage collection to happen every %d packages."), GarbageCollectionFrequency);
+		}
 	}
 
 	if ((bShouldBuildLighting || bShouldBuildReflectionCaptures) && !bExplicitPackages)
@@ -282,9 +296,11 @@ int32 UResavePackagesCommandlet::InitializeResaveParameters( const TArray<FStrin
 			PackageFilter |= NORMALIZE_ExcludeNonDeveloperPackages;
 		}
 
-		bool bAnyFound = NormalizePackageNames(Unused, PackageNames, *FString::Printf(TEXT("*%s"), *FPackageName::GetAssetPackageExtension()), 
-											   PackageFilter);
-		bAnyFound = NormalizePackageNames(Unused, PackageNames, *FString::Printf(TEXT("*%s"), *FPackageName::GetMapPackageExtension()), PackageFilter) || bAnyFound;
+		const FString AssetSearch = TEXT("*") + FPackageName::GetAssetPackageExtension();
+		const FString MapSearch = TEXT("*") + FPackageName::GetMapPackageExtension();
+
+		bool bAnyFound = NormalizePackageNames(Unused, PackageNames, AssetSearch, PackageFilter);
+		bAnyFound |= NormalizePackageNames(Unused, PackageNames, MapSearch, PackageFilter);
 		
 		if (!bAnyFound)
 		{
@@ -488,6 +504,55 @@ int32 UResavePackagesCommandlet::InitializeResaveParameters( const TArray<FStrin
 	return 0;
 }
 
+void UResavePackagesCommandlet::ParseSourceControlOptions(const TArray<FString>& Tokens)
+{
+	if (SourceControlQueue == nullptr)
+	{
+		return; // No point parsing the options if we don't have anything enabled
+	}
+
+	int32 QueuedPackageFlushLimit = INDEX_NONE;
+	int64 QueueFileSizeFlushLimit = INDEX_NONE;
+
+	for (const FString& CurrentSwitch : Tokens)
+	{
+		if (FParse::Value(*CurrentSwitch, TEXT("BatchPackageLimit="), QueuedPackageFlushLimit))
+		{
+			if (QueuedPackageFlushLimit >= 0)
+			{
+				UE_LOG(LogContentCommandlet, Display, TEXT("Setting source control batches to be limited to %d package(s) at a time."), QueuedPackageFlushLimit);
+				SourceControlQueue->SetMaxNumQueuedPackages(QueuedPackageFlushLimit);
+			}
+			else
+			{
+				// Negative values mean we will not flush the source control batch based on the number of packages
+				UE_LOG(LogContentCommandlet, Display, TEXT("Setting source control batches to have no package limit!"));
+			}
+		}
+		else  if (FParse::Value(*CurrentSwitch, TEXT("BatchFileSizeLimit="), QueueFileSizeFlushLimit))
+		{
+			const int64 SizeLimit = TNumericLimits<int64>::Max() / (1024 * 1024);
+			if (QueueFileSizeFlushLimit > SizeLimit)
+			{
+				UE_LOG(LogContentCommandlet, Display, TEXT("-BatchFileSizeLimit=%lld is too large! The max value allowed is %lld, clamping..."), QueueFileSizeFlushLimit, SizeLimit);
+				QueueFileSizeFlushLimit = SizeLimit;
+			}
+
+			if (QueueFileSizeFlushLimit >= 0)
+			{
+				UE_LOG(LogContentCommandlet, Display, TEXT("Setting source control batches to be limited to %lld MB."), QueueFileSizeFlushLimit);
+
+				SourceControlQueue->SetMaxTemporaryFileTotalSize(QueueFileSizeFlushLimit);
+			}
+			else
+			{
+				// Negative values mean we will not flush the source control batch based on the disk space taken by temp files
+				UE_LOG(LogContentCommandlet, Display, TEXT("Setting source control batches to have no disk space limit!"));
+			}
+		}
+	}
+}
+
 bool UResavePackagesCommandlet::ShouldSkipPackage(const FString& Filename)
 {
 	return false;
@@ -495,6 +560,8 @@ bool UResavePackagesCommandlet::ShouldSkipPackage(const FString& Filename)
 
 void UResavePackagesCommandlet::LoadAndSaveOnePackage(const FString& Filename)
 {
+	TRACE_CPUPROFILER_EVENT_SCOPE(UResavePackagesCommandlet::LoadAndSaveOnePackage);
+
 	// Check to see if a derived commandlet wants to skip this package for one reason or another
 	if (ShouldSkipPackage(Filename))
 	{
@@ -611,7 +678,12 @@ void UResavePackagesCommandlet::LoadAndSaveOnePackage(const FString& Filename)
 			GStaticMeshPackageNameToRebuild = FName(*FPackageName::FilenameToLongPackageName(Filename));
 
 			// Assert if package couldn't be opened so we have no chance of messing up saving later packages.
-			UPackage* Package = LoadPackage( NULL, *Filename, 0 );
+			UPackage* Package = nullptr;
+			{
+				TRACE_CPUPROFILER_EVENT_SCOPE(UResavePackagesCommandlet::LoadAndSaveOnePackage::LoadPackage);
+				Package = LoadPackage(NULL, *Filename, 0);
+			}
+
 			if (Package == NULL)
 			{
 				if (bCanIgnoreFails == true)
@@ -658,13 +730,10 @@ void UResavePackagesCommandlet::LoadAndSaveOnePackage(const FString& Filename)
 			VerboseMessage(TEXT("Post PerformAdditionalOperations"));
 
 			// Check for any special per object operations
-			for( FObjectIterator ObjectIt; ObjectIt; ++ObjectIt )
+			ForEachObjectWithOuter(Package, [this, &bSavePackage](UObject* Object)
 			{
-				if( ObjectIt->IsIn( Package ) )
-				{
-					PerformAdditionalOperations( *ObjectIt, bSavePackage );
-				}
-			}
+				PerformAdditionalOperations(Object, bSavePackage);
+			});
 			
 			VerboseMessage(TEXT("Post PerformAdditionalOperations Loop"));
 
@@ -717,9 +786,11 @@ void UResavePackagesCommandlet::LoadAndSaveOnePackage(const FString& Filename)
 				}
 				else
 				{
-					// check to see if we need to check this package out
-					if ( bAutoCheckOut )
+					// Check to see if we need to check this package out (but do not check out here if SourceControlQueue is enabled)
+					const bool bAttemptCheckoutNow = bAutoCheckOut && SourceControlQueue == nullptr;
+					if (bAttemptCheckoutNow)
 					{
+						TRACE_CPUPROFILER_EVENT_SCOPE(UResavePackagesCommandlet::LoadAndSaveOnePackage::AutoCheckOut);
 						if( bIsReadOnly )
 						{
 							VerboseMessage(TEXT("Pre ForceGetStatus1"));
@@ -751,9 +822,12 @@ void UResavePackagesCommandlet::LoadAndSaveOnePackage(const FString& Filename)
 						}
 					}
 
-					// so now we need to see if we actually were able to check this file out
-					// if the file is still read only then we failed and need to emit an error and go to the next package
-					if (IFileManager::Get().IsReadOnly( *Filename ) == true)
+					// Update the readonly state now that source control has had a chance to run
+					bIsReadOnly = IFileManager::Get().IsReadOnly(*Filename);
+
+					// If we tried to check out the file but it is still read only then we failed and 
+					// need to emit an error and go to the next package			
+					if (bAttemptCheckoutNow && bIsReadOnly)
 					{
 						if (bSkipCheckedOutFiles)
 						{
@@ -775,12 +849,29 @@ void UResavePackagesCommandlet::LoadAndSaveOnePackage(const FString& Filename)
 
 					const static bool bKeepPackageGUIDOnSave = FParse::Param(FCommandLine::Get(), TEXT("KeepPackageGUIDOnSave"));
 					ESaveFlags SaveFlags = bKeepPackageGUIDOnSave ? SAVE_KeepGUID : SAVE_None;
-					if( SavePackageHelper(Package, Filename, RF_Standalone, GWarn, nullptr, SaveFlags) )
+					
+					if (bIsReadOnly == false || SourceControlQueue == nullptr)
 					{
-						PackagesResaved++;
-						if (Verbosity == VERY_VERBOSE)
+						TRACE_CPUPROFILER_EVENT_SCOPE(UResavePackagesCommandlet::LoadAndSaveOnePackage::SavePackage);
+
+						if (SavePackageHelper(Package, Filename, RF_Standalone, GWarn, nullptr, SaveFlags))
+						{
+							PackagesResaved++;
+							if (Verbosity == VERY_VERBOSE)
+							{
+								UE_LOG(LogContentCommandlet, Display, TEXT("Correctly saved:  [%s]."), *Filename);
+							}
+						}
+					}
+					else
+					{
+						// The target file is still read only and we have SourceControlQueue enabled so we
+						// need to save to a temporary file first, and then queue the result.
+						const FString TempFilename = CreateTempFilename();
+
+						if (SavePackageHelper(Package, TempFilename, RF_Standalone, GWarn, nullptr, SaveFlags))
 						{
-							UE_LOG(LogContentCommandlet, Display, TEXT("Correctly saved:  [%s]."), *Filename );
+							SourceControlQueue->QueueCheckoutAndReplaceOperation(Filename, TempFilename, Package);
 						}
 					}
 				}
@@ -806,6 +897,8 @@ void UResavePackagesCommandlet::LoadAndSaveOnePackage(const FString& Filename)
 
 void UResavePackagesCommandlet::DeleteOnePackage(const FString& Filename)
 {
+	TRACE_CPUPROFILER_EVENT_SCOPE(UResavePackagesCommandlet::DeleteOnePackage);
+
 	bool bIsReadOnly = IFileManager::Get().IsReadOnly(*Filename);
 
 	if (bVerifyContent)
@@ -822,6 +915,14 @@ void UResavePackagesCommandlet::DeleteOnePackage(const FString& Filename)
 		return;
 	}
 
+	if (SourceControlQueue != nullptr)
+	{
+		// All files (read only and non read only) need to query their source control status so
+		// they should all be queued if avaliable.
+		SourceControlQueue->QueueDeleteOperation(Filename);
+		return;
+	}
+
 	FString PackageName;
 	FPackageName::TryConvertFilenameToLongPackageName(Filename, PackageName);
 
@@ -849,6 +950,8 @@ void UResavePackagesCommandlet::DeleteOnePackage(const FString& Filename)
 		UE_LOG(LogContentCommandlet, Display, TEXT("Deleting '%s' from source control..."), *Filename);
 		SourceControlProvider.Execute(ISourceControlOperation::Create<FDelete>(), PackageFilename);
 
+		PackagesDeleted++;
+
 		FilesToSubmit.AddUnique(Filename);
 	}
 	else if (SourceControlState.IsValid() && SourceControlState->CanCheckout())
@@ -856,6 +959,8 @@ void UResavePackagesCommandlet::DeleteOnePackage(const FString& Filename)
 		UE_LOG(LogContentCommandlet, Display, TEXT("Deleting '%s' from source control..."), *Filename);
 		SourceControlProvider.Execute(ISourceControlOperation::Create<FDelete>(), PackageFilename);
 
+		PackagesDeleted++;
+
 		FilesToSubmit.AddUnique(Filename);
 	}
 	else if (SourceControlState.IsValid() && SourceControlState->IsCheckedOutOther())
@@ -865,7 +970,11 @@ void UResavePackagesCommandlet::DeleteOnePackage(const FString& Filename)
 	else if (SourceControlState.IsValid() && !SourceControlState->IsSourceControlled())
 	{
 		UE_LOG(LogContentCommandlet, Warning, TEXT("'%s' is not in source control, attempting to delete from disk..."), *Filename);
-		if (!IFileManager::Get().Delete(*Filename, false, true))
+		if (IFileManager::Get().Delete(*Filename, false, true) == true)
+		{
+			PackagesDeleted++;
+		}
+		else
 		{
 			UE_LOG(LogContentCommandlet, Warning, TEXT("  ... failed to delete from disk."), *Filename);
 		}
@@ -873,7 +982,11 @@ void UResavePackagesCommandlet::DeleteOnePackage(const FString& Filename)
 	else
 	{
 		UE_LOG(LogContentCommandlet, Warning, TEXT("'%s' is in an unknown source control state, attempting to delete from disk..."), *Filename);
-		if (!IFileManager::Get().Delete(*Filename, false, true))
+		if (IFileManager::Get().Delete(*Filename, false, true)== true)
+		{
+			PackagesDeleted++;
+		}
+		else
 		{
 			UE_LOG(LogContentCommandlet, Warning, TEXT("  ... failed to delete from disk."), *Filename);
 		}
@@ -882,6 +995,8 @@ void UResavePackagesCommandlet::DeleteOnePackage(const FString& Filename)
 
 int32 UResavePackagesCommandlet::Main( const FString& Params )
 {
+	TRACE_CPUPROFILER_EVENT_SCOPE(UResavePackagesCommandlet);
+
 	const TCHAR* Parms = *Params;
 	TArray<FString> Tokens;
 	ParseCommandLine(Parms, Tokens, Switches);
@@ -899,6 +1014,8 @@ int32 UResavePackagesCommandlet::Main( const FString& Params )
 	bOnlySaveDirtyPackages = Switches.Contains(TEXT("OnlySaveDirtyPackages"));
 	/** if we should auto checkout packages that need to be saved**/
 	bAutoCheckOut = Switches.Contains(TEXT("AutoCheckOutPackages")) || Switches.Contains(TEXT("AutoCheckOut"));
+	/** if we should batch source control operations*/
+	bBatchSourceControl = Switches.Contains(TEXT("BatchSourceControl"));
 	/** if we should simply skip checked out files rather than error-ing out */
 	bSkipCheckedOutFiles = Switches.Contains(TEXT("SkipCheckedOutPackages"));
 	/** if we should auto checkin packages that were checked out**/
@@ -1036,6 +1153,7 @@ int32 UResavePackagesCommandlet::Main( const FString& Params )
 	int32 GCIndex = 0;
 	PackagesConsideredForResave = 0;
 	PackagesResaved = 0;
+	PackagesDeleted = 0;
 
 	// allow for an option to restart at a given package name (in case it dies during a run, etc)
 	bool bCanProcessPackage = true;
@@ -1059,6 +1177,32 @@ int32 UResavePackagesCommandlet::Main( const FString& Params )
 		GShaderCompilingManager->ProcessAsyncResults(true, false);
 	}
 
+	if (bBatchSourceControl)
+	{
+		// Convert the commandlets verbosity to that of FQueuedSourceControlOperations
+		FQueuedSourceControlOperations::EVerbosity LogVerbosity = FQueuedSourceControlOperations::EVerbosity::All;
+		switch (Verbosity)
+		{
+		case VERY_VERBOSE:
+			LogVerbosity = FQueuedSourceControlOperations::EVerbosity::All;
+			break;
+		case INFORMATIVE:
+			LogVerbosity = FQueuedSourceControlOperations::EVerbosity::Info;
+			break;
+		case ONLY_ERRORS:
+			LogVerbosity = FQueuedSourceControlOperations::EVerbosity::ErrorsOnly;
+			break;
+		default:
+			UE_LOG(LogContentCommandlet, Warning, TEXT("Unknown verbosity to pass to FQueuedSourceControlOperations!"));
+		}
+
+		SourceControlQueue = MakePimpl<FQueuedSourceControlOperations>(LogVerbosity);	
+		ParseSourceControlOptions(Switches);
+	}
+
+	// Make sure any remaining temp files from previous runs are removed
+	CleanTempFiles();
+
 	// Iterate over all packages.
 	for( int32 PackageIndex = 0; PackageIndex < PackageNames.Num(); PackageIndex++ )
 	{
@@ -1085,6 +1229,12 @@ int32 UResavePackagesCommandlet::Main( const FString& Params )
 		// Load and save this package
 		LoadAndSaveOnePackage(Filename);
 
+		// Check if we need to flush any source control operations yet
+		if (SourceControlQueue != nullptr)
+		{
+			SourceControlQueue->FlushPendingOperations(false);
+		}
+
 		// Tick shader compiler if we are rendering
 		if(IsAllowCommandletRendering() && GShaderCompilingManager)
 		{
@@ -1098,6 +1248,22 @@ int32 UResavePackagesCommandlet::Main( const FString& Params )
 			break;
 		}
 	}
+	
+	if (SourceControlQueue != nullptr)
+	{
+		// Flush any remaining source control operations
+		SourceControlQueue->FlushPendingOperations(true);
+
+		// Fix up the stats
+		PackagesDeleted += SourceControlQueue->GetNumDeletedFiles();
+		PackagesResaved += SourceControlQueue->GetNumReplacedFiles();
+
+		// Add files that the source control queue modified to the list of files that 
+		// will need to be submitted to source control
+		FilesToSubmit.Append(SourceControlQueue->GetModifiedFiles());
+
+		SourceControlQueue.Reset();
+	}
 
 	// Force a directory watcher and asset registry tick
 	FDirectoryWatcherModule& DirectoryWatcherModule = FModuleManager::Get().LoadModuleChecked<FDirectoryWatcherModule>(TEXT("DirectoryWatcher"));
@@ -1144,6 +1310,7 @@ int32 UResavePackagesCommandlet::Main( const FString& Params )
 
 	UE_LOG(LogContentCommandlet, Display, TEXT("[REPORT] %d/%d packages were considered for resaving"), PackagesConsideredForResave, PackageNames.Num());
 	UE_LOG(LogContentCommandlet, Display, TEXT("[REPORT] %d/%d packages were resaved"), PackagesResaved, PackagesConsideredForResave);
+	UE_LOG(LogContentCommandlet, Display, TEXT("[REPORT] %d/%d packages were deleted"), PackagesDeleted, PackagesConsideredForResave);
 
 
 	return 0;
@@ -1916,6 +2083,29 @@ void UResavePackagesCommandlet::VerboseMessage(const FString& Message)
 	}
 }
 
+FString UResavePackagesCommandlet::CreateTempFilename()
+{
+	return FPaths::CreateTempFilename(*GetTempFilesDirectory());
+}
+
+void UResavePackagesCommandlet::CleanTempFiles()
+{
+	TRACE_CPUPROFILER_EVENT_SCOPE(UResavePackagesCommandlet::CleanTempFiles);
+	const FString DirPath = GetTempFilesDirectory();
+
+	UE_LOG(LogContentCommandlet, Display, TEXT("Cleaning temp file directory"), *DirPath);
+
+	if (!IFileManager::Get().DeleteDirectory(*DirPath, false, true))
+	{
+		UE_LOG(LogContentCommandlet, Warning, TEXT("Failed to clean temp file directory:  %s"), *DirPath);
+	}
+}
+
+FString UResavePackagesCommandlet::GetTempFilesDirectory()
+{
+	return FPaths::ProjectSavedDir() / TEXT("Temp/ResavePackages");
+}
+
 /*-----------------------------------------------------------------------------
 	UWrangleContent.
 -----------------------------------------------------------------------------*/
@@ -2954,4 +3144,3 @@ int32 UListStaticMeshesImportedFromSpeedTreesCommandlet::Main(const FString& Par
 	}
 	return 0;
 }
-
diff --git a/Engine/Source/Editor/UnrealEd/Private/Commandlets/CookCommandlet.cpp b/Engine/Source/Editor/UnrealEd/Private/Commandlets/CookCommandlet.cpp
--- a/Engine/Source/Editor/UnrealEd/Private/Commandlets/CookCommandlet.cpp
+++ b/Engine/Source/Editor/UnrealEd/Private/Commandlets/CookCommandlet.cpp
@@ -51,6 +51,8 @@ DEFINE_LOG_CATEGORY_STATIC(LogCookCommandlet, Log, All);
 #include "IAnalyticsProviderET.h"
 #include "AnalyticsET.h"
 
+ASSETREGISTRY_API void DumpBuildDependencyTrackerData();
+
 namespace DetailedCookStats
 {
 	FString CookProject;
@@ -342,6 +344,8 @@ namespace DetailedCookStats
 				UE_LOG(LogCookCommandlet, Display, TEXT("%-34s  %16.2f  %21.2f  %12d  %12.2f"), *Stat.AssetType, Stat.TotalTimeSec, Stat.GameThreadTimeSec, Stat.AssetsBuilt, Stat.SizeMB);
 			}
 		}
+
+		DumpBuildDependencyTrackerData();
 	}
 }
 #endif
diff --git a/Engine/Source/Editor/UnrealEd/Private/CookOnTheFlyServer.cpp b/Engine/Source/Editor/UnrealEd/Private/CookOnTheFlyServer.cpp
--- a/Engine/Source/Editor/UnrealEd/Private/CookOnTheFlyServer.cpp
+++ b/Engine/Source/Editor/UnrealEd/Private/CookOnTheFlyServer.cpp
@@ -5514,6 +5514,8 @@ void UCookOnTheFlyServer::SaveGlobalShaderMapFiles(const TArrayView<const ITarge
 
 		FString OutputDir = GetSandboxDirectory(RecompileData.PlatformName);
 
+		UE_LOG(LogCook, Display, TEXT("Checking global shaders for platform %s"), *RecompileData.PlatformName);
+
 		RecompileShadersForRemote
 			(RecompileData.PlatformName, 
 			RecompileData.ShaderPlatform == -1 ? SP_NumPlatforms : (EShaderPlatform)RecompileData.ShaderPlatform, //-V547
diff --git a/Engine/Source/Editor/UnrealEd/Private/EditorEngine.cpp b/Engine/Source/Editor/UnrealEd/Private/EditorEngine.cpp
--- a/Engine/Source/Editor/UnrealEd/Private/EditorEngine.cpp
+++ b/Engine/Source/Editor/UnrealEd/Private/EditorEngine.cpp
@@ -230,6 +230,9 @@
 #include "UObject/UnrealType.h"
 #include "Factories/TextureFactory.h"
 #include "Engine/TextureCube.h"
+#include "Misc/PackageAccessTracking.h"
+#include "Misc/PackageAccessTrackingOps.h"
+
 #if WITH_CHAOS
 #include "ChaosSolversModule.h"
 #endif
@@ -4461,6 +4464,7 @@ FSavePackageResultStruct UEditorEngine::Save( UPackage* InOuter, UObject* InBase
 				 FSavePackageContext* SavePackageContext)
 {
 	TRACE_CPUPROFILER_EVENT_SCOPE(UEditorEngine::Save);
+	UE_TRACK_REFERENCING_PACKAGE_SCOPED(InOuter, PackageAccessTrackingOps::NAME_Save); // Needs to be here in addition to UPackage::Save so that InitializePhysicsSceneForSaveIfNecessary and OnPreSaveWorld have appropriate referencing package info
 
 	FScopedSlowTask SlowTask(100, FText(), bSlowTask);
 
diff --git a/Engine/Source/Editor/UnrealEd/Private/Factories/DataTableFactory.cpp b/Engine/Source/Editor/UnrealEd/Private/Factories/DataTableFactory.cpp
--- a/Engine/Source/Editor/UnrealEd/Private/Factories/DataTableFactory.cpp
+++ b/Engine/Source/Editor/UnrealEd/Private/Factories/DataTableFactory.cpp
@@ -200,7 +200,7 @@ UObject* UDataTableFactory::FactoryCreateNew(UClass* Class, UObject* InParent, F
 		DataTable = MakeNewDataTable(InParent, Name, Flags);
 		if (DataTable)
 		{
-			DataTable->RowStruct = const_cast<UScriptStruct*>(Struct);
+			DataTable->RowStruct = const_cast<UScriptStruct*>(ToRawPtr(Struct));
 		}
 	}
 	return DataTable;
diff --git a/Engine/Source/Editor/UnrealEd/Private/Kismet2/Kismet2.cpp b/Engine/Source/Editor/UnrealEd/Private/Kismet2/Kismet2.cpp
--- a/Engine/Source/Editor/UnrealEd/Private/Kismet2/Kismet2.cpp
+++ b/Engine/Source/Editor/UnrealEd/Private/Kismet2/Kismet2.cpp
@@ -449,7 +449,7 @@ UBlueprint* FKismetEditorUtilities::CreateBlueprint(UClass* ParentClass, UObject
 
 		NewBP->SimpleConstructionScript = NewObject<USimpleConstructionScript>(NewClass);
 		NewBP->SimpleConstructionScript->SetFlags(RF_Transactional);
-		NewBP->LastEditedDocuments.Add(NewBP->SimpleConstructionScript);
+		NewBP->LastEditedDocuments.Add(ToRawPtr(NewBP->SimpleConstructionScript));
 
 		UEdGraph* UCSGraph = FBlueprintEditorUtils::CreateNewGraph(NewBP, UEdGraphSchema_K2::FN_UserConstructionScript, UEdGraph::StaticClass(), UEdGraphSchema_K2::StaticClass());
 		FBlueprintEditorUtils::AddFunctionGraph(NewBP, UCSGraph, /*bIsUserCreated=*/ false, AActor::StaticClass());
diff --git a/Engine/Source/Editor/UnrealEd/Private/ThumbnailRendering/SlateBrushThumbnailRenderer.cpp b/Engine/Source/Editor/UnrealEd/Private/ThumbnailRendering/SlateBrushThumbnailRenderer.cpp
--- a/Engine/Source/Editor/UnrealEd/Private/ThumbnailRendering/SlateBrushThumbnailRenderer.cpp
+++ b/Engine/Source/Editor/UnrealEd/Private/ThumbnailRendering/SlateBrushThumbnailRenderer.cpp
@@ -24,7 +24,7 @@ void USlateBrushThumbnailRenderer::Draw(UObject* Object, int32 X, int32 Y, uint3
 
 		// Draw the background checkboard pattern
 		const int32 CheckerDensity = 8;
-		auto* Checker = UThumbnailManager::Get().CheckerboardTexture;
+		auto Checker = ToRawPtr(UThumbnailManager::Get().CheckerboardTexture);
 		Canvas->DrawTile(
 			0.0f, 0.0f, Width, Height,							// Dimensions
 			0.0f, 0.0f, CheckerDensity, CheckerDensity,			// UVs
diff --git a/Engine/Source/Editor/UnrealEd/Private/ThumbnailRendering/TextureThumbnailRenderer.cpp b/Engine/Source/Editor/UnrealEd/Private/ThumbnailRendering/TextureThumbnailRenderer.cpp
--- a/Engine/Source/Editor/UnrealEd/Private/ThumbnailRendering/TextureThumbnailRenderer.cpp
+++ b/Engine/Source/Editor/UnrealEd/Private/ThumbnailRendering/TextureThumbnailRenderer.cpp
@@ -98,7 +98,7 @@ void UTextureThumbnailRenderer::Draw(UObject* Object, int32 X, int32 Y, uint32 W
 		{
 			// If using alpha, draw a checkerboard underneath first.
 			const int32 CheckerDensity = 8;
-			auto* Checker = UThumbnailManager::Get().CheckerboardTexture;
+			auto Checker = UThumbnailManager::Get().CheckerboardTexture;
 			Canvas->DrawTile(
 				0.0f, 0.0f, Width, Height,							// Dimensions
 				0.0f, 0.0f, CheckerDensity, CheckerDensity,			// UVs
diff --git a/Engine/Source/Editor/UnrealEd/Private/Tools/UEdMode.cpp b/Engine/Source/Editor/UnrealEd/Private/Tools/UEdMode.cpp
--- a/Engine/Source/Editor/UnrealEd/Private/Tools/UEdMode.cpp
+++ b/Engine/Source/Editor/UnrealEd/Private/Tools/UEdMode.cpp
@@ -251,7 +251,7 @@ void UEdMode::RegisterTool(TSharedPtr<FUICommandInfo> UICommand, FString ToolIde
 	ToolsContext->ToolManager->RegisterToolType(ToolIdentifier, Builder);
 	CommandList->MapAction(UICommand,
 		FExecuteAction::CreateUObject(ToolsContext.Get(), &UEdModeInteractiveToolsContext::StartTool, ToolIdentifier),
-		FCanExecuteAction::CreateUObject(ToolsContext->ToolManager, &UInteractiveToolManager::CanActivateTool, EToolSide::Mouse, ToolIdentifier),
+		FCanExecuteAction::CreateUObject(ToRawPtr(ToolsContext->ToolManager), &UInteractiveToolManager::CanActivateTool, EToolSide::Mouse, ToolIdentifier),
 		FIsActionChecked::CreateUObject(ToolsContext.Get(), &UEdModeInteractiveToolsContext::IsToolBuilderActive, EToolSide::Mouse, Builder),
 		EUIActionRepeatMode::RepeatDisabled);
 
diff --git a/Engine/Source/Editor/VREditor/VREditorAvatarActor.cpp b/Engine/Source/Editor/VREditor/VREditorAvatarActor.cpp
--- a/Engine/Source/Editor/VREditor/VREditorAvatarActor.cpp
+++ b/Engine/Source/Editor/VREditor/VREditorAvatarActor.cpp
@@ -464,13 +464,13 @@ void AVREditorAvatarActor::TickManually( const float DeltaTime )
 			UPostProcessComponent* PostProcess = PostProcessComponent;
 			if ( bShouldDrawWorldMovementPostProcess )
 			{
-				PostProcess->AddOrUpdateBlendable( WorldMovementPostProcessMaterial );
+				PostProcess->AddOrUpdateBlendable( ToRawPtr(WorldMovementPostProcessMaterial) );
 				bIsDrawingWorldMovementPostProcess = true;
 			}
 			else
 			{
 				bIsDrawingWorldMovementPostProcess = false;
-				PostProcess->Settings.RemoveBlendable( WorldMovementPostProcessMaterial );
+				PostProcess->Settings.RemoveBlendable( ToRawPtr(WorldMovementPostProcessMaterial) );
 			}
 		}
 
diff --git a/Engine/Source/Programs/NativePointerUpgradeTool/NativePointerUpgradeTool.Build.cs b/Engine/Source/Programs/NativePointerUpgradeTool/NativePointerUpgradeTool.Build.cs
new file mode 100644
--- /dev/null
+++ b/Engine/Source/Programs/NativePointerUpgradeTool/NativePointerUpgradeTool.Build.cs
@@ -0,0 +1,16 @@
+// Copyright Epic Games, Inc. All Rights Reserved.
+
+using UnrealBuildTool;
+
+public class NativePointerUpgradeTool : ModuleRules
+{
+	public NativePointerUpgradeTool(ReadOnlyTargetRules Target) : base(Target)
+	{
+		PublicIncludePaths.Add("Runtime/Launch/Public");
+
+		PrivateIncludePaths.Add("Runtime/Launch/Private");		// For LaunchEngineLoop.cpp include
+
+		PrivateDependencyModuleNames.Add("Core");
+		PrivateDependencyModuleNames.Add("Projects");
+	}
+}
diff --git a/Engine/Source/Programs/NativePointerUpgradeTool/NativePointerUpgradeTool.Target.cs b/Engine/Source/Programs/NativePointerUpgradeTool/NativePointerUpgradeTool.Target.cs
new file mode 100644
--- /dev/null
+++ b/Engine/Source/Programs/NativePointerUpgradeTool/NativePointerUpgradeTool.Target.cs
@@ -0,0 +1,24 @@
+// Copyright Epic Games, Inc. All Rights Reserved.
+
+using UnrealBuildTool;
+using System.Collections.Generic;
+
+[SupportedPlatforms(UnrealPlatformClass.Desktop)]
+public class NativePointerUpgradeToolTarget : TargetRules
+{
+	public NativePointerUpgradeToolTarget(TargetInfo Target) : base(Target)
+	{
+		Type = TargetType.Program;
+		LinkType = TargetLinkType.Modular;
+		LaunchModuleName = "NativePointerUpgradeTool";
+
+		bBuildDeveloperTools			= false;
+		bUseMallocProfiler				= false;
+		bBuildWithEditorOnlyData		= false;
+		bCompileAgainstEngine			= false;
+		bCompileAgainstCoreUObject		= false;
+		bCompileAgainstApplicationCore	= false;
+		bUsesSlate = false;
+		bIsBuildingConsoleApplication	= true;
+	}
+}
diff --git a/Engine/Source/Programs/NativePointerUpgradeTool/Private/NativePointerUpgradeTool.cpp b/Engine/Source/Programs/NativePointerUpgradeTool/Private/NativePointerUpgradeTool.cpp
new file mode 100644
--- /dev/null
+++ b/Engine/Source/Programs/NativePointerUpgradeTool/Private/NativePointerUpgradeTool.cpp
@@ -0,0 +1,618 @@
+// Copyright Epic Games, Inc. All Rights Reserved.
+
+
+#include "CoreMinimal.h"
+#include "RequiredProgramMainCPPInclude.h"
+#include "CoreTypes.h"
+#include "Containers/Array.h"
+#include "Containers/Map.h"
+#include "Logging/LogMacros.h"
+#include "Misc/CString.h"
+#include "Misc/ScopeExit.h"
+#include "Misc/StringBuilder.h"
+#include "String/LexFromString.h"
+
+DEFINE_LOG_CATEGORY_STATIC(LogNativePointerUpgradeTool, Log, All);
+
+IMPLEMENT_APPLICATION(NativePointerUpgradeTool, "NativePointerUpgradeTool");
+
+//////////////////////////////////////////////////////////////////////////
+
+enum class ENativePointerUpgradeBehaviorFlags
+{
+	None,
+	PreviewOnly = 1,
+};
+ENUM_CLASS_FLAGS(ENativePointerUpgradeBehaviorFlags)
+
+class FNativePointerUpgrader
+{
+public:
+
+	FNativePointerUpgrader(const FString& InSCCCommand, ENativePointerUpgradeBehaviorFlags InBehaviorFlags) : SCCCommand(InSCCCommand), BehaviorFlags(InBehaviorFlags)
+	{
+	}
+
+	struct FMemberDeclaration
+	{
+		int32 LineNumber = 0;
+		int32 LineCount = 0;
+		FString TypeName;
+
+		bool operator==(const FMemberDeclaration& Other) const
+		{
+			return (LineNumber == Other.LineNumber) && (TypeName == Other.TypeName);
+		}
+	};
+	using FNativePointerUpgradeList = TMap<FString, TArray<FMemberDeclaration>>;
+
+	TValueOrError<FNativePointerUpgradeList, FString> FillUpgradeList(const FString& LogFilename);
+	bool TryAllUpgrades(const FNativePointerUpgradeList& UpgradeList);
+
+	int32 GetTotalUpgradesFound() const { return TotalUpgradesFound; }
+	int32 GetTotalUpgradesPerformed() const { return TotalUpgradesPerformed; }
+
+private:
+	enum class EByteOrderMarkerType
+	{
+		None,
+		UTF8,
+		UCS2LittleEndian,
+		UCS2BigEndian,
+	};
+
+	struct FPendingWrite
+	{
+		EByteOrderMarkerType BOM;
+		TArray<FString> Contents;
+	};
+	using FPendingWrites = TMap<FString, FPendingWrite>;
+	bool TryUpgradesInFile(const FString& FilenameToUpgrade, TConstArrayView<FMemberDeclaration> MemberDeclarations, FPendingWrites& PendingWrites);
+	bool TryFlushPendingWrites(FPendingWrites& PendingWrites);
+
+	// Using these methods instead of the ones in FFileHelper to ensure that we retain the encoding and BOM of the original file as well as the newlines
+	static bool LoadFileToString(FString& OutResult, EByteOrderMarkerType& OutBOM, const TCHAR* InFilename, uint32 ReadFlags = 0);
+	static bool LoadFileToStringArrayWithNewlines(TArray<FString>& OutResult, EByteOrderMarkerType& OutBOM, const TCHAR* InFilename, uint32 ReadFlags = 0);
+	static bool SaveStringToFile(const FStringView& InString, EByteOrderMarkerType InBOM, const TCHAR* InFilename, uint32 WriteFlags = 0);
+	static bool SaveStringArrayWithNewlinesToFile(const TArray<FString>& InStringArray, EByteOrderMarkerType InBOM, const TCHAR* InFilename, uint32 WriteFlags = 0);
+
+	static int32 CountLines(const FStringView Input);
+
+
+	int32 TotalUpgradesFound = 0;
+	int32 TotalUpgradesPerformed = 0;
+	FString SCCCommand;
+	ENativePointerUpgradeBehaviorFlags BehaviorFlags;
+};
+
+FStringView ParseBracketedStringSegment(const FString& InString, int32& InOutCurrentIndex, const TCHAR* BracketStart, const TCHAR* BracketEnd)
+{
+	int32 ParsedStartIndex = InString.Find(BracketStart, ESearchCase::IgnoreCase, ESearchDir::FromStart, InOutCurrentIndex);
+	if (ParsedStartIndex == INDEX_NONE)
+	{
+		return FStringView();
+	}
+	ParsedStartIndex += FCString::Strlen(BracketStart);
+	int32 ParsedEndIndex = InString.Find(BracketEnd, ESearchCase::IgnoreCase, ESearchDir::FromStart, ParsedStartIndex);
+	if (ParsedEndIndex == INDEX_NONE)
+	{
+		return FStringView();
+	}
+	InOutCurrentIndex = ParsedEndIndex + FCString::Strlen(BracketEnd);
+	return FStringView(*InString + ParsedStartIndex, ParsedEndIndex - ParsedStartIndex);
+}
+
+TValueOrError<FNativePointerUpgrader::FNativePointerUpgradeList, FString> FNativePointerUpgrader::FillUpgradeList(const FString& LogFilename)
+{
+	FNativePointerUpgradeList UpgradeList;
+	const TCHAR Preamble[] = TEXT("Native pointer usage in member declaration detected in '");
+	TArray<FString> NativePointerUpgradeEntries;
+	if (!FFileHelper::LoadFileToStringArrayWithPredicate(NativePointerUpgradeEntries, *LogFilename, [&Preamble](const FString& Line) { return Line.Contains(Preamble);  }))
+	{
+		return MakeError(FString::Printf(TEXT("Unable to load UHT log: %s"), *LogFilename));
+	}
+
+	constexpr int32 PreambleLen = UE_ARRAY_COUNT(Preamble) - 1;
+	for (const FString& Entry : NativePointerUpgradeEntries)
+	{
+		int32 CurrentIndex = 0;
+		FStringView FilenameView = ParseBracketedStringSegment(Entry, CurrentIndex, Preamble, TEXT("'"));
+		if (FilenameView.IsEmpty())
+		{
+			return MakeError(FString::Printf(TEXT("Failed to parse filename from native pointer member upgrade statement: %s"), *Entry));
+		}
+
+		FStringView LineNumberView = ParseBracketedStringSegment(Entry, CurrentIndex, TEXT(", line "), TEXT(" "));
+		if (LineNumberView.IsEmpty())
+		{
+			return MakeError(FString::Printf(TEXT("Failed to parse line number from native pointer member upgrade statement: %s"), *Entry));
+		}
+
+		FStringView TypeNameView = ParseBracketedStringSegment(Entry, CurrentIndex, TEXT("[["), TEXT("]]"));
+		if (TypeNameView.IsEmpty())
+		{
+			return MakeError(FString::Printf(TEXT("Failed to parse type name from native pointer member upgrade statement: %s"), *Entry));
+		}
+
+		int32 LineNumber = -1;
+		LexFromString(LineNumber, LineNumberView);
+		if (LineNumber < 1)
+		{
+			return MakeError(FString::Printf(TEXT("Failed to convert line number to integer from native pointer member upgrade statement: %s"), *Entry));
+		}
+
+		if (!TypeNameView.EndsWith(TEXT("*")))
+		{
+			return MakeError(FString::Printf(TEXT("Failed to find expected asterisk at end of native pointer member upgrade statement: %s"), *Entry));
+		}
+
+		TArray<FMemberDeclaration>& DeclarationArray = UpgradeList.FindOrAdd(FString(FilenameView));
+		FMemberDeclaration NewDecl;
+		NewDecl.LineNumber = LineNumber;
+		NewDecl.TypeName = TypeNameView;
+		NewDecl.TypeName = NewDecl.TypeName.ReplaceEscapedCharWithChar();
+		NewDecl.LineCount = CountLines(NewDecl.TypeName);
+		int32 DeclArraySize = DeclarationArray.Num();
+		DeclarationArray.AddUnique(NewDecl);
+		if (DeclArraySize != DeclarationArray.Num())
+		{
+			TotalUpgradesFound++;
+		}
+	}
+
+	return MakeValue(MoveTemp(UpgradeList));
+}
+
+// The following requirements need to be kept in mind while upgrading source files:
+// 1) If the file section we're updating isn't what we expect it to be, we shouldn't attempt the upgrade
+// 2) We should handle multi-line type declarations
+// 3) We should handle multiple type replacements on the same line
+// 4) We should retain the encoding of the original file
+// 5) We should retain the newlines of the original file
+// 6) We should attempt to retain the indentation (before and after type declaration) of any lines we modify
+bool FNativePointerUpgrader::TryUpgradesInFile(const FString& FilenameToUpgrade, TConstArrayView<FMemberDeclaration> MemberDeclarations, FPendingWrites& PendingWrites)
+{
+	EByteOrderMarkerType OriginalBOM;
+	TArray<FString> FileContentsToUpgrade;
+	if (!LoadFileToStringArrayWithNewlines(FileContentsToUpgrade, OriginalBOM, *FilenameToUpgrade))
+	{
+		UE_LOG(LogNativePointerUpgradeTool, Error, TEXT("Unable to load file to upgrade: %s"), *FilenameToUpgrade);
+		return true;
+	}
+
+	bool bChangesMadeToFileContents = false;
+	for (const FMemberDeclaration& MemberDeclaration : MemberDeclarations)
+	{
+		if (((MemberDeclaration.LineNumber-MemberDeclaration.LineCount) < 0) || (MemberDeclaration.LineNumber > FileContentsToUpgrade.Num()))
+		{
+			UE_LOG(LogNativePointerUpgradeTool, Error, TEXT("Member declaration line '%d' is out of range for the lines (%d) in file '%s' "), MemberDeclaration.LineNumber, FileContentsToUpgrade.Num(), *FilenameToUpgrade);
+			continue;
+		}
+
+		FString LinesToUpgrade;
+		for (int32 LineIndex = MemberDeclaration.LineNumber-MemberDeclaration.LineCount; LineIndex < MemberDeclaration.LineNumber; ++LineIndex)
+		{
+			LinesToUpgrade.Append(FileContentsToUpgrade[LineIndex]);
+		}
+		FString OriginalLinesToUpgrade = LinesToUpgrade;
+
+		// Chop off the trailing asterisk in the member type declaration
+		FString TemplateArg = MemberDeclaration.TypeName.LeftChop(1);
+		// Any trailing whitespace that was after the asterisk should not go in the template argument, but be appended after the template argument
+		// This only applies to space (' ') and tab ('\t') that have non-whitespace before them on the line, and should not include newlines.
+		// This is to avoid having end template bracket ('>') get put on a line that may have a line comment ahead of it ("//").
+		int32 TrailingWhitespaceCount = 0;
+		int32 CharIndex;
+		for (CharIndex = TemplateArg.Len()-1; (CharIndex >= 0) && (TemplateArg[CharIndex] == TCHAR(' ') || TemplateArg[CharIndex] == TCHAR('\t')); --CharIndex)
+		{
+			++TrailingWhitespaceCount;
+		}
+		if ((CharIndex == 0) || (TemplateArg[CharIndex-1] == TEXT('\n')) || (TemplateArg[CharIndex-1] == TEXT('\r')))
+		{
+			// The trailing whitespace represented the entire lead of the line the end template character will be on.  Put the indentation inside the template arg instead of outside of it.
+			TrailingWhitespaceCount = 0;
+		}
+		FString TrailingWhitespace = TemplateArg.Right(TrailingWhitespaceCount);
+		TemplateArg.LeftChopInline(TrailingWhitespaceCount);
+		FString UpgradedTypeName = FString::Printf(TEXT("TObjectPtr<%s>%s"), *TemplateArg, *TrailingWhitespace);
+		if (LinesToUpgrade.ReplaceInline(*MemberDeclaration.TypeName, *UpgradedTypeName) < 1)
+		{
+			if (!LinesToUpgrade.Contains(*UpgradedTypeName))
+			{
+				UE_LOG(LogNativePointerUpgradeTool, Warning, TEXT("%s(%d): Type name string '%s' not found for upgrading, skipping upgrading this line:\n\t%s"), *FilenameToUpgrade, MemberDeclaration.LineNumber, *MemberDeclaration.TypeName, *LinesToUpgrade.TrimStartAndEnd());
+			}
+			continue;
+		}
+
+		TArray<FTextRange> UpgradedTextRanges;
+		FTextRange::CalculateLineRangesFromString(LinesToUpgrade, UpgradedTextRanges);
+
+		for (int32 LineIndex = 0; LineIndex < MemberDeclaration.LineCount; ++LineIndex)
+		{
+			int32 BeginIndex = UpgradedTextRanges[LineIndex].BeginIndex;
+			int32 EndIndex = LineIndex < UpgradedTextRanges.Num()-1 ? UpgradedTextRanges[LineIndex+1].BeginIndex : UpgradedTextRanges[LineIndex].EndIndex;
+			FString UpgradedLine = LinesToUpgrade.Mid(BeginIndex, EndIndex - BeginIndex);
+			FileContentsToUpgrade[MemberDeclaration.LineNumber - MemberDeclaration.LineCount + LineIndex] = UpgradedLine;
+		}
+
+		bChangesMadeToFileContents = true;
+		TotalUpgradesPerformed++;
+		UE_LOG(LogNativePointerUpgradeTool, Display, TEXT("%s(%d): %s\n\tFrom:\t%s\n\tTo:\t%s"), *FilenameToUpgrade, MemberDeclaration.LineNumber, EnumHasAllFlags(BehaviorFlags, ENativePointerUpgradeBehaviorFlags::PreviewOnly) ? TEXT("Previewing upgrade") : TEXT("Upgrading"), *OriginalLinesToUpgrade.TrimStartAndEnd(), *LinesToUpgrade.TrimStartAndEnd());
+	}
+
+	if (bChangesMadeToFileContents)
+	{
+		FPendingWrite& NewWrite = PendingWrites.Add(FilenameToUpgrade);
+		NewWrite.BOM = OriginalBOM;
+		NewWrite.Contents.Append(MoveTemp(FileContentsToUpgrade));
+	}
+
+	return true;
+}
+
+bool FNativePointerUpgrader::TryFlushPendingWrites(FPendingWrites& PendingWrites)
+{
+	if (PendingWrites.IsEmpty())
+	{
+		return true;
+	}
+
+	// Perform a batched SCC command operation
+	if (!SCCCommand.IsEmpty())
+	{
+		TStringBuilder<2048> CompositeFilenamesBuilder;
+		for (const FPendingWrites::ElementType& PendingWrite : PendingWrites)
+		{
+			CompositeFilenamesBuilder.Append(TEXT('"'));
+			CompositeFilenamesBuilder.Append(PendingWrite.Key);
+			CompositeFilenamesBuilder.Append(TEXT("\" "));
+		}
+
+		FString FileSCCCommand(SCCCommand);
+		FileSCCCommand.ReplaceInline(TEXT("{Filenames}"), *CompositeFilenamesBuilder);
+		if (EnumHasAllFlags(BehaviorFlags, ENativePointerUpgradeBehaviorFlags::PreviewOnly))
+		{
+			UE_LOG(LogNativePointerUpgradeTool, Display, TEXT("Would perform SCC command: %s"), *FileSCCCommand);
+		}
+		else
+		{
+			UE_LOG(LogNativePointerUpgradeTool, Display, TEXT("Perform SCC command: %s"), *FileSCCCommand);
+		}
+
+		if (!EnumHasAllFlags(BehaviorFlags, ENativePointerUpgradeBehaviorFlags::PreviewOnly))
+		{
+			FString Params = FCommandLine::RemoveExeName(*FileSCCCommand);
+			FString Cmd = FileSCCCommand.LeftChop(Params.Len()).TrimEnd();
+
+			FProcHandle ProcHandle = FPlatformProcess::CreateProc(*Cmd, *Params, false, false, false, nullptr, 0, nullptr, nullptr);
+			if (!ProcHandle.IsValid())
+			{
+				UE_LOG(LogNativePointerUpgradeTool, Error, TEXT("Failed to perform SCC command: %s"), *FileSCCCommand);
+				return false;
+			}
+
+			FPlatformProcess::WaitForProc(ProcHandle);
+			int32 RetCode = -1;
+			FPlatformProcess::GetProcReturnCode(ProcHandle, &RetCode);
+			FPlatformProcess::CloseProc(ProcHandle);
+
+			if (RetCode != 0)
+			{
+				UE_LOG(LogNativePointerUpgradeTool, Error, TEXT("SCC command returned an unexpected return code: %d"), RetCode);
+				return false;
+			}
+		}
+	}
+
+	for (const FPendingWrites::ElementType& PendingWrite : PendingWrites)
+	{
+		if (EnumHasAllFlags(BehaviorFlags, ENativePointerUpgradeBehaviorFlags::PreviewOnly))
+		{
+			UE_LOG(LogNativePointerUpgradeTool, Display, TEXT("File would be modified and saved: %s"), *PendingWrite.Key);
+		}
+		else
+		{
+			UE_LOG(LogNativePointerUpgradeTool, Display, TEXT("Saving modified file: %s"), *PendingWrite.Key);
+			if (!SaveStringArrayWithNewlinesToFile(PendingWrite.Value.Contents, PendingWrite.Value.BOM, *PendingWrite.Key))
+			{
+				UE_LOG(LogNativePointerUpgradeTool, Error, TEXT("Failed to save upgraded file: %s"), *PendingWrite.Key);
+				return false;
+			}
+		}
+	}
+
+	PendingWrites.Empty();
+	return true;
+}
+
+bool FNativePointerUpgrader::TryAllUpgrades(const FNativePointerUpgradeList& UpgradeList)
+{
+	FPendingWrites PendingWrites;
+	for (const TPair<FString, TArray<FMemberDeclaration>>& FileToUpgrade : UpgradeList)
+	{
+		if (!TryUpgradesInFile(FileToUpgrade.Key, FileToUpgrade.Value, PendingWrites))
+		{
+			return false;
+		}
+
+		if (PendingWrites.Num() >= 20)
+		{
+			if (!TryFlushPendingWrites(PendingWrites))
+			{
+				return false;
+			}
+			check(PendingWrites.IsEmpty());
+		}
+	}
+
+	if (!TryFlushPendingWrites(PendingWrites))
+	{
+		return false;
+	}
+	return true;
+}
+
+bool FNativePointerUpgrader::LoadFileToString(FString& OutResult, EByteOrderMarkerType& OutBOM, const TCHAR* InFilename, uint32 ReadFlags)
+{
+	TUniquePtr<FArchive> Reader(IFileManager::Get().CreateFileReader(InFilename, ReadFlags));
+	if (!Reader)
+	{
+		return false;
+	}
+
+	FScopedLoadingState ScopedLoadingState(*Reader->GetArchiveName());
+
+	int64 Size = Reader->TotalSize();
+	if (!Size)
+	{
+		OutBOM = EByteOrderMarkerType::None;
+		OutResult.Empty();
+		return true;
+	}
+
+	uint8* OriginalBuffer = (uint8*)FMemory::Malloc(Size);
+	uint8* Buffer = OriginalBuffer;
+	ON_SCOPE_EXIT
+	{
+		FMemory::Free(OriginalBuffer);
+	};
+
+	Reader->Serialize(Buffer, Size);
+	if (!Reader->Close())
+	{
+		return false;
+	}
+
+	TArray<TCHAR>& ResultArray = OutResult.GetCharArray();
+	ResultArray.Empty();
+
+	bool bIsUnicode = false;
+	if (Size >= 2 && !(Size & 1) && Buffer[0] == 0xff && Buffer[1] == 0xfe)
+	{
+		// Unicode Intel byte order. Less 1 for the FFFE header, additional 1 for null terminator.
+		OutBOM = EByteOrderMarkerType::UCS2LittleEndian;
+		ResultArray.AddUninitialized(Size / 2);
+		for (int32 i = 0; i < (Size / 2) - 1; i++)
+		{
+			ResultArray[i] = CharCast<TCHAR>((UCS2CHAR)((uint16)Buffer[i * 2 + 2] + (uint16)Buffer[i * 2 + 3] * 256));
+		}
+		bIsUnicode = true;
+	}
+	else if (Size >= 2 && !(Size & 1) && Buffer[0] == 0xfe && Buffer[1] == 0xff)
+	{
+		// Unicode non-Intel byte order. Less 1 for the FFFE header, additional 1 for null terminator.
+		OutBOM = EByteOrderMarkerType::UCS2BigEndian;
+		ResultArray.AddUninitialized(Size / 2);
+		for (int32 i = 0; i < (Size / 2) - 1; i++)
+		{
+			ResultArray[i] = CharCast<TCHAR>((UCS2CHAR)((uint16)Buffer[i * 2 + 3] + (uint16)Buffer[i * 2 + 2] * 256));
+		}
+		bIsUnicode = true;
+	}
+	else
+	{
+		if (Size >= 3 && Buffer[0] == 0xef && Buffer[1] == 0xbb && Buffer[2] == 0xbf)
+		{
+			// Skip over UTF-8 BOM if there is one
+			OutBOM = EByteOrderMarkerType::UTF8;
+			Buffer += 3;
+			Size -= 3;
+		}
+		else
+		{
+			OutBOM = EByteOrderMarkerType::None;
+		}
+
+		int32 Length = FUTF8ToTCHAR_Convert::ConvertedLength(reinterpret_cast<const ANSICHAR*>(Buffer), Size);
+		ResultArray.AddUninitialized(Length + 1); // +1 for the null terminator
+		FUTF8ToTCHAR_Convert::Convert(ResultArray.GetData(), ResultArray.Num(), reinterpret_cast<const ANSICHAR*>(Buffer), Size);
+		ResultArray[Length] = TEXT('\0');
+	}
+
+	if (ResultArray.Num() == 1)
+	{
+		// If it's only a zero terminator then make the result actually empty
+		ResultArray.Empty();
+	}
+	else
+	{
+		// Else ensure null terminator is present
+		ResultArray.Last() = 0;
+
+		if (bIsUnicode)
+		{
+			// Inline combine any surrogate pairs in the data when loading into a UTF-32 string
+			StringConv::InlineCombineSurrogates(OutResult);
+		}
+	}
+
+	return true;
+}
+
+bool FNativePointerUpgrader::LoadFileToStringArrayWithNewlines(TArray<FString>& OutResult, EByteOrderMarkerType& OutBOM, const TCHAR* InFilename, uint32 ReadFlags)
+{
+	FString FileContents;
+	if (!LoadFileToString(FileContents, OutBOM, InFilename, ReadFlags))
+	{
+		return false;
+	}
+
+	TArray<FTextRange> LineRanges;
+	FTextRange::CalculateLineRangesFromString(FileContents, LineRanges);
+	OutResult.Reset(LineRanges.Num());
+
+	for (int32 LineIndex = 0; LineIndex < LineRanges.Num(); ++LineIndex)
+	{
+		int32 BeginIndex = LineRanges[LineIndex].BeginIndex;
+		int32 EndIndex = (LineIndex < (LineRanges.Num() - 1)) ? LineRanges[LineIndex+1].BeginIndex : LineRanges[LineIndex].EndIndex;
+		OutResult.Add(FileContents.Mid(BeginIndex, EndIndex - BeginIndex));
+	}
+
+	return true;
+}
+
+bool FNativePointerUpgrader::SaveStringToFile(const FStringView& InString, EByteOrderMarkerType InBOM, const TCHAR* InFilename, uint32 WriteFlags)
+{
+	// max size of the string is a UCS2CHAR for each character and some UNICODE magic 
+	TUniquePtr<FArchive> Ar = TUniquePtr<FArchive>(IFileManager::Get().CreateFileWriter(InFilename, WriteFlags));
+	if (!Ar)
+	{
+		return false;
+	}
+
+	if (InString.IsEmpty())
+	{
+		return true;
+	}
+
+	if (InBOM == EByteOrderMarkerType::UTF8)
+	{
+		UTF8CHAR UTF8BOM[] = {0xEF, 0xBB, 0xBF};
+		Ar->Serialize(&UTF8BOM, UE_ARRAY_COUNT(UTF8BOM) * sizeof(UTF8CHAR));
+
+		FTCHARToUTF8 UTF8String(InString.GetData(), InString.Len());
+		Ar->Serialize((UTF8CHAR*)UTF8String.Get(), UTF8String.Length() * sizeof(UTF8CHAR));
+	}
+	else if ((InBOM == EByteOrderMarkerType::UCS2LittleEndian) || (InBOM == EByteOrderMarkerType::UCS2BigEndian))
+	{
+		Ar->SetByteSwapping(PLATFORM_LITTLE_ENDIAN ? InBOM == EByteOrderMarkerType::UCS2BigEndian : InBOM == EByteOrderMarkerType::UCS2LittleEndian);
+		
+		UTF16CHAR BOM = UNICODE_BOM;
+		Ar->Serialize(&BOM, sizeof(UTF16CHAR));
+
+		// Note: This is a no-op on platforms that are using a 16-bit TCHAR
+		FTCHARToUTF16 UTF16String(InString.GetData(), InString.Len());
+		Ar->Serialize((UTF16CHAR*)UTF16String.Get(), UTF16String.Length() * sizeof(UTF16CHAR));
+	}
+	else
+	{
+		auto Src = StringCast<ANSICHAR>(InString.GetData(), InString.Len());
+		Ar->Serialize((ANSICHAR*)Src.Get(), Src.Length() * sizeof(ANSICHAR));
+	}
+
+	// Always explicitly close to catch errors from flush/close
+	Ar->Close();
+
+	return !Ar->IsError() && !Ar->IsCriticalError();
+}
+
+bool FNativePointerUpgrader::SaveStringArrayWithNewlinesToFile(const TArray<FString>& InStringArray, EByteOrderMarkerType InBOM, const TCHAR* InFilename, uint32 WriteFlags)
+{
+	FString ComposedFileContents = FString::Join(InStringArray, TEXT(""));
+	return SaveStringToFile(ComposedFileContents, InBOM, InFilename, WriteFlags);
+}
+
+int32 FNativePointerUpgrader::CountLines(const FStringView Input)
+{
+	int32 LineCount = 0;
+	int32 LineBeginIndex = 0;
+
+	// Loop through splitting at new-lines
+	const TCHAR* const InputStart = Input.GetData();
+	const TCHAR* const InputEnd = InputStart + Input.Len();
+	for (const TCHAR* CurrentChar = InputStart; CurrentChar != InputEnd; ++CurrentChar)
+	{
+		// Handle a chain of \r\n slightly differently to stop the FChar::IsLinebreak adding two separate new-lines
+		const bool bIsWindowsNewLine = ((CurrentChar + 1) != InputEnd) && (*CurrentChar == '\r' && *(CurrentChar + 1) == '\n');
+		if (bIsWindowsNewLine || FChar::IsLinebreak(*CurrentChar))
+		{
+			++LineCount;
+
+			if (bIsWindowsNewLine)
+			{
+				++CurrentChar; // skip the \n of the \r\n chain
+			}
+			LineBeginIndex = UE_PTRDIFF_TO_INT32(CurrentChar - InputStart) + 1; // The next line begins after the end of the current line
+		}
+	}
+
+	// Process any remaining string after the last new-line
+	if (LineBeginIndex <= Input.Len())
+	{
+		++LineCount;
+	}
+
+	return LineCount;
+}
+
+//////////////////////////////////////////////////////////////////////////
+
+INT32_MAIN_INT32_ARGC_TCHAR_ARGV()
+{
+	FString CmdLine = FCommandLine::BuildFromArgV(nullptr, ArgC, ArgV, nullptr);
+	GEngineLoop.PreInit(*CmdLine);
+
+	// Make sure the engine is properly cleaned up whenever we exit this function
+	ON_SCOPE_EXIT
+	{
+		FEngineLoop::AppPreExit();
+		FEngineLoop::AppExit();
+	};
+
+	FString UHTLogFilename;
+	if (CmdLine.Len() && **CmdLine != TEXT('-'))
+	{
+		const TCHAR* CmdLinePtr = *CmdLine;
+		UHTLogFilename = FParse::Token(CmdLinePtr, false);
+	}
+
+	if (UHTLogFilename.IsEmpty())
+	{
+		UHTLogFilename = FPaths::ConvertRelativePathToFull(FPlatformProcess::BaseDir(), TEXT("../../../Engine/Programs/UnrealHeaderTool/Saved/Logs/UnrealHeaderTool.log"));
+	}
+	
+	ENativePointerUpgradeBehaviorFlags BehaviorFlags = ENativePointerUpgradeBehaviorFlags::None;
+	FString SCCCommand;
+	FParse::Value(*CmdLine, TEXT("-SCCCommand="), SCCCommand);
+	if (FParse::Param(*CmdLine, TEXT("n")) || FParse::Param(*CmdLine, TEXT("PREVIEW")))
+	{
+		BehaviorFlags |= ENativePointerUpgradeBehaviorFlags::PreviewOnly;
+	}
+
+	FNativePointerUpgrader NativePointerUpgrader(SCCCommand, BehaviorFlags);
+
+	TValueOrError<FNativePointerUpgrader::FNativePointerUpgradeList, FString> UpgradeListResult =
+		NativePointerUpgrader.FillUpgradeList(UHTLogFilename);
+
+	if (!UpgradeListResult.IsValid())
+	{
+		UE_LOG(LogNativePointerUpgradeTool, Error, TEXT("Upgrade log parsing error: %s"), *UpgradeListResult.GetError());
+		return 1;
+	}
+
+	if (!NativePointerUpgrader.TryAllUpgrades(UpgradeListResult.GetValue()))
+	{
+		return 1;
+	}
+
+	if (EnumHasAllFlags(BehaviorFlags, ENativePointerUpgradeBehaviorFlags::PreviewOnly))
+	{
+		UE_LOG(LogNativePointerUpgradeTool, Display, TEXT("Previewed upgrade successfully.  %d upgrades found; %d upgrades would be performed"), NativePointerUpgrader.GetTotalUpgradesFound(), NativePointerUpgrader.GetTotalUpgradesPerformed());
+	}
+	else
+	{
+		UE_LOG(LogNativePointerUpgradeTool, Display, TEXT("Upgraded successfully.  %d upgrades found; %d upgrades performed"), NativePointerUpgrader.GetTotalUpgradesFound(), NativePointerUpgrader.GetTotalUpgradesPerformed());
+	}
+
+	return 0;
+}
diff --git a/Engine/Source/Programs/UnrealHeaderTool/Private/CodeGenerator.cpp b/Engine/Source/Programs/UnrealHeaderTool/Private/CodeGenerator.cpp
--- a/Engine/Source/Programs/UnrealHeaderTool/Private/CodeGenerator.cpp
+++ b/Engine/Source/Programs/UnrealHeaderTool/Private/CodeGenerator.cpp
@@ -1423,6 +1423,28 @@ void FNativeClassHeaderGenerator::PropertyNew(FOutputDevice& DeclOut, FOutputDev
 		return;
 	}
 
+	if (FObjectPtrProperty* TypedProp = CastField<FObjectPtrProperty>(Prop))
+	{
+		DeclOut.Logf(TEXT("%sstatic const UE4CodeGen_Private::FObjectPtrPropertyParams %s;\r\n"), DeclSpaces, *NameWithoutScope);
+
+		Out.Logf(
+			TEXT("%sconst UE4CodeGen_Private::FObjectPtrPropertyParams %s = { %s, %s, (EPropertyFlags)0x%016llx, UE4CodeGen_Private::EPropertyGenFlags::Object | UE4CodeGen_Private::EPropertyGenFlags::ObjectPtr, %s, %s, %s, %s, %s };%s\r\n"),
+			Spaces,
+			Name,
+			*PropName,
+			*PropNotifyFunc,
+			PropFlags,
+			FPropertyObjectFlags,
+			*ArrayDim,
+			OffsetStr,
+			*GetSingletonNameFuncAddr(TypedProp->PropertyClass, OutReferenceGatherers.UniqueCrossModuleReferences, false),
+			*MetaDataParams,
+			*PropTag
+		);
+
+		return;
+	}
+
 	if (FSoftObjectProperty* TypedProp = CastField<FSoftObjectProperty>(Prop))
 	{
 		DeclOut.Logf(TEXT("%sstatic const UE4CodeGen_Private::FSoftObjectPropertyParams %s;\r\n"), DeclSpaces, *NameWithoutScope);
