Prototype Wrapped Object Pointers (ObjectHandle/ObjectPtr)

Fix error encountered when using NativePointerUpgradeTool on UnrealEditor project and encountering a source file that is UTF8 encoded but lacking a BOM.

diff --git a/Engine/Source/Programs/NativePointerUpgradeTool/Private/NativePointerUpgradeTool.cpp b/Engine/Source/Programs/NativePointerUpgradeTool/Private/NativePointerUpgradeTool.cpp
--- a/Engine/Source/Programs/NativePointerUpgradeTool/Private/NativePointerUpgradeTool.cpp
+++ b/Engine/Source/Programs/NativePointerUpgradeTool/Private/NativePointerUpgradeTool.cpp
@@ -506,8 +506,13 @@ bool FNativePointerUpgrader::SaveStringToFile(const FStringView& InString, EByte
 	}
 	else
 	{
-		auto Src = StringCast<ANSICHAR>(InString.GetData(), InString.Len());
-		Ar->Serialize((ANSICHAR*)Src.Get(), Src.Length() * sizeof(ANSICHAR));
+		// If there is no BOM specified, write to UTF8 without a BOM.  If the contents fit within ASCII, they will be
+		// written as ASCII.  If they don't fit, they'll be written as multibyte code points.  This is necessary as
+		// a file was encountered that is UTF8 without any BOM.  The reading code handles this correctly,
+		// so we want to ensure that the writing code doesn't substitute "bogus chars" (?) for any non-ASCII characters
+		// at the time of writing.
+		FTCHARToUTF8 UTF8String(InString.GetData(), InString.Len());
+		Ar->Serialize((UTF8CHAR*)UTF8String.Get(), UTF8String.Length() * sizeof(UTF8CHAR));
 	}
 
 	// Always explicitly close to catch errors from flush/close
